Sound Cue Templates

diff --git a/Engine/Plugins/Runtime/SoundCueTemplates/SoundCueTemplates.uplugin b/Engine/Plugins/Runtime/SoundCueTemplates/SoundCueTemplates.uplugin
new file mode 100644
--- /dev/null
+++ b/Engine/Plugins/Runtime/SoundCueTemplates/SoundCueTemplates.uplugin
@@ -0,0 +1,23 @@
+{
+	"FileVersion": 3,
+	"Version": 1,
+	"VersionName": "1.0",
+	"FriendlyName": "Sound Cue Templates",
+	"Description": "Collection of SoundCue Templates, which provide rapid design of common audio design workflows.",
+	"Category": "Audio",
+	"CreatedBy": "Epic Games, Inc.",
+	"CreatedByURL": "",
+	"DocsURL": "",
+	"MarketplaceURL": "",
+	"SupportURL": "",
+	"CanContainContent": false,
+	"IsBetaVersion": false,
+	"Installed": false,
+	"Modules": [
+		{
+			"Name": "SoundCueTemplates",
+			"Type": "Runtime",
+			"LoadingPhase": "Default"
+		}
+	]
+}
\ No newline at end of file
diff --git a/Engine/Plugins/Runtime/SoundCueTemplates/Source/SoundCueTemplates/Private/SoundCueContainer.cpp b/Engine/Plugins/Runtime/SoundCueTemplates/Source/SoundCueTemplates/Private/SoundCueContainer.cpp
new file mode 100644
--- /dev/null
+++ b/Engine/Plugins/Runtime/SoundCueTemplates/Source/SoundCueTemplates/Private/SoundCueContainer.cpp
@@ -0,0 +1,269 @@
+// Copyright 1998-2019 Epic Games, Inc. All Rights Reserved.
+#include "SoundCueContainer.h"
+
+#include "Sound/SoundCue.h"
+#include "Sound/SoundCueTemplate.h"
+#include "Sound/SoundWave.h"
+#include "Sound/SoundNode.h"
+#include "Sound/SoundNodeConcatenator.h"
+#include "Sound/SoundNodeMixer.h"
+#include "Sound/SoundNodeModulator.h"
+#include "Sound/SoundNodeQualityLevel.h"
+#include "Sound/SoundNodeRandom.h"
+#include "Sound/SoundNodeWavePlayer.h"
+#include "SoundCueTemplateSettings.h"
+#include "SoundCueTemplatesModule.h"
+
+#if WITH_EDITORONLY_DATA
+#include "DetailCategoryBuilder.h"
+#include "DetailLayoutBuilder.h"
+#include "Sound/AudioSettings.h"
+#include "SoundCueGraph/SoundCueGraphNode.h"
+#include "SoundCueGraph/SoundCueGraph.h"
+#include "SoundCueGraph/SoundCueGraphNode_Root.h"
+#include "SoundCueGraph/SoundCueGraphSchema.h"
+#endif // WITH_EDITORONLY_DATA
+
+#define LOCTEXT_NAMESPACE "SoundCueContainer"
+USoundCueContainer::USoundCueContainer(const FObjectInitializer& ObjectInitializer)
+	: Super(ObjectInitializer)
+{
+#if WITH_EDITORONLY_DATA
+	bLooping = false;
+
+	ContainerType = ESoundContainerType::Randomize;
+
+	PitchModulation.X = 0.95;
+	PitchModulation.Y = 1.05;
+
+	// Roughly -3dB to 0 dB
+	VolumeModulation.X = 0.70f;
+	VolumeModulation.Y = 1.00f;
+#endif // WITH_EDITORONLY_DATA
+}
+
+#if WITH_EDITORONLY_DATA
+void USoundCueContainer::OnRebuildGraph(USoundCue& SoundCue) const
+{
+	const USoundCueTemplateSettings* TemplateSettings = GetDefault<USoundCueTemplateSettings>();
+	if (!TemplateSettings)
+	{
+		return;
+	}
+
+	const UAudioSettings* AudioSettings = GetDefault<UAudioSettings>();
+	if (!AudioSettings)
+	{
+		return;
+	}
+
+	if (TemplateSettings->GetQualityLevelSettingsNum() != AudioSettings->GetQualityLevelSettingsNum())
+	{
+		UE_LOG(SoundCueTemplates, Warning, TEXT("Rebuild SoundCueTemplate 'SoundCueContainer' failed: "
+			"Quality settings mismatch between Audio Settings & SoundCueTemplate settings."))
+		return;
+	}
+
+	if (!Variations.Num())
+	{
+		return;
+	}
+
+	USoundNodeModulator& ModulatorNode = ConstructSoundNodeRoot<USoundNodeModulator>(SoundCue);
+	ModulatorNode.PitchMin = PitchModulation.X;
+	ModulatorNode.PitchMax = PitchModulation.Y;
+	ModulatorNode.VolumeMin = VolumeModulation.X;
+	ModulatorNode.VolumeMax = VolumeModulation.Y;
+
+	int32 Column = 1;
+	int32 Row = 0;
+	const int32 InputPinIndex = 0;
+	USoundNodeQualityLevel& QualityNode = USoundCueTemplate::ConstructSoundNodeChild<USoundNodeQualityLevel>(SoundCue, &ModulatorNode, Column, Row, InputPinIndex);
+	Column++;
+
+	// Optimization to just link to single wave player and omit randomizer if just one variation provided.
+	USoundNodeWavePlayer* WavePlayerNode = nullptr;
+
+	int32 NumQualitySettings = AudioSettings->GetQualityLevelSettingsNum();
+
+	// Get max variations
+	int32 TotalWavePlayers = 0;
+	for (int32 i = 0; i < NumQualitySettings; ++i)
+	{
+		const FSoundCueTemplateQualitySettings& QualitySettings = TemplateSettings->GetQualityLevelSettings(i);
+		const int32 MaxVariations = GetMaxVariations(QualitySettings);
+		TotalWavePlayers = FMath::Max(TotalWavePlayers, MaxVariations);
+	}
+
+	// Prune variations to only include valid (non-null) references
+	TArray<USoundWave*> VariationArray = Variations.Array();
+	for (int32 i = Variations.Num() -1 ; i >= 0; --i)
+	{
+		if (!VariationArray[i])
+		{
+			VariationArray.RemoveAtSwap(i, 1, false);
+		}
+	}
+
+	// Trim total to the size of possible variations & create wave player nodes
+	TotalWavePlayers = FMath::Min(TotalWavePlayers, VariationArray.Num());
+	TArray<USoundNodeWavePlayer*> WavePlayers;
+	const int32 WavePlayerColumn = 4;
+	for (int32 i = 0; i < TotalWavePlayers; ++i)
+	{
+		const int32 WavePlayerRow = i;
+		USoundNodeWavePlayer* WavePlayer = &USoundCueTemplate::ConstructSoundNodeChild<USoundNodeWavePlayer>(SoundCue, nullptr, WavePlayerColumn, WavePlayerRow);
+		WavePlayers.Add(WavePlayer);
+		WavePlayer->SetSoundWave(VariationArray[i]);
+		WavePlayer->bLooping = bLooping;
+	}
+
+	// Map of max channels to container to avoid duplicating container nodes
+	TMap<int32, USoundNode*> ContainerMap;
+
+	int32 ContainerRow = 0;
+	for (int32 i = 0; i < NumQualitySettings; ++i)
+	{
+		const int32 PinIndex = i;
+		const FAudioQualitySettings& AudioQualitySettings = AudioSettings->GetQualityLevelSettings(i);
+		const FSoundCueTemplateQualitySettings& QualitySettings = TemplateSettings->GetQualityLevelSettings(i);
+		const int32 MaxVariations = GetMaxVariations(QualitySettings);
+		if (WavePlayers.Num() == 1)
+		{
+			USoundCueTemplate::AddSoundNodeChild(QualityNode, *WavePlayers[0], PinIndex, &QualitySettings.DisplayName);
+		}
+		else
+		{
+			USoundNode* ContainerNode = ContainerMap.FindRef(MaxVariations);
+			if (ContainerNode)
+			{
+				USoundCueTemplate::AddSoundNodeChild(QualityNode, *ContainerNode, PinIndex, &QualitySettings.DisplayName);
+			}
+			else
+			{
+				switch (ContainerType)
+				{
+				case ESoundContainerType::Concatenate:
+				{
+					ContainerNode = &USoundCueTemplate::ConstructSoundNodeChild<USoundNodeConcatenator>(
+						SoundCue,
+						&QualityNode,
+						Column,
+						ContainerRow,
+						i,
+						&QualitySettings.DisplayName);
+				}
+				break;
+
+				case ESoundContainerType::Mix:
+				{
+					ContainerNode = &USoundCueTemplate::ConstructSoundNodeChild<USoundNodeMixer>(
+						SoundCue,
+						&QualityNode,
+						Column,
+						ContainerRow,
+						i,
+						&QualitySettings.DisplayName);
+				}
+				break;
+
+				case ESoundContainerType::Randomize:
+				default:
+				{
+					ContainerNode = &USoundCueTemplate::ConstructSoundNodeChild<USoundNodeRandom>(
+						SoundCue,
+						&QualityNode,
+						Column,
+						ContainerRow,
+						i,
+						&QualitySettings.DisplayName);
+				}
+				break;
+				}
+				check(ContainerNode);
+
+				// Add wave players and connect to random node
+				for (int32 j = 0; j < MaxVariations && j < WavePlayers.Num(); ++j)
+				{
+					USoundCueTemplate::AddSoundNodeChild(*ContainerNode, *WavePlayers[j], j);
+				}
+
+				ContainerMap.Add(MaxVariations, ContainerNode);
+				ContainerRow++;
+			}
+		}
+	}
+}
+
+int32 USoundCueContainer::GetMaxVariations(const FSoundCueTemplateQualitySettings& QualitySettings) const
+{
+	int32 MaxVariations = TNumericLimits<int32>::Max();
+
+	switch (ContainerType)
+	{
+		case ESoundContainerType::Concatenate:
+		{
+			return QualitySettings.MaxConcatenatedVariations;
+		}
+		break;
+
+		case ESoundContainerType::Mix:
+		{
+			return QualitySettings.MaxMixVariations;
+		}
+		break;
+
+		case ESoundContainerType::Randomize:
+		default:
+		{
+			return QualitySettings.MaxRandomizedVariations;
+		}
+		break;
+	}
+
+	return MaxVariations;
+}
+
+TSet<FName>& USoundCueContainer::GetWhitelistedCategories()
+{
+	static TSet<FName> Categories;
+	if (Categories.Num() == 0)
+	{
+		Categories.Add(FName(TEXT("Modulation")));
+		Categories.Add(FName(TEXT("Variation")));
+	}
+
+	return Categories;
+}
+
+TSharedRef<IDetailCustomization> FSoundCueContainerDetailCustomization::MakeInstance()
+{
+	return MakeShareable(new FSoundCueContainerDetailCustomization);
+}
+
+void FSoundCueContainerDetailCustomization::CustomizeDetails(IDetailLayoutBuilder& DetailLayout)
+{
+	TArray<FName> CategoryNames;
+	DetailLayout.GetCategoryNames(CategoryNames);
+
+	const TSet<FName>& Whitelist = USoundCueContainer::GetWhitelistedCategories();
+	for (FName CategoryName : CategoryNames)
+	{
+		if (!Whitelist.Contains(CategoryName))
+		{
+			DetailLayout.HideCategory(CategoryName);
+		}
+	}
+}
+
+void FSoundCueContainerDetailCustomization::Register(FPropertyEditorModule& PropertyModule)
+{
+	PropertyModule.RegisterCustomClassLayout("SoundCueContainer", FOnGetDetailCustomizationInstance::CreateStatic(&FSoundCueContainerDetailCustomization::MakeInstance));
+}
+
+void FSoundCueContainerDetailCustomization::Unregister(FPropertyEditorModule& PropertyModule)
+{
+	PropertyModule.UnregisterCustomClassLayout("SoundCueContainer");
+}
+#endif // WITH_EDITORONLY_DATA
+#undef LOCTEXT_NAMESPACE
diff --git a/Engine/Plugins/Runtime/SoundCueTemplates/Source/SoundCueTemplates/Private/SoundCueDistanceCrossfade.cpp b/Engine/Plugins/Runtime/SoundCueTemplates/Source/SoundCueTemplates/Private/SoundCueDistanceCrossfade.cpp
new file mode 100644
--- /dev/null
+++ b/Engine/Plugins/Runtime/SoundCueTemplates/Source/SoundCueTemplates/Private/SoundCueDistanceCrossfade.cpp
@@ -0,0 +1,133 @@
+// Copyright 1998-2019 Epic Games, Inc. All Rights Reserved.
+
+
+#include "SoundCueDistanceCrossfade.h"
+
+
+#include "Sound/SoundCue.h"
+#include "Sound/SoundWave.h"
+#include "Sound/SoundNode.h"
+#include "Sound/SoundNodeModulator.h"
+#include "Sound/SoundNodeWavePlayer.h"
+
+#if WITH_EDITOR
+#include "DetailCategoryBuilder.h"
+#include "DetailLayoutBuilder.h"
+#include "Sound/AudioSettings.h"
+#include "SoundCueGraph/SoundCueGraphNode.h"
+#include "SoundCueGraph/SoundCueGraph.h"
+#include "SoundCueGraph/SoundCueGraphNode_Root.h"
+#include "SoundCueGraph/SoundCueGraphSchema.h"
+#endif // WITH_EDITOR
+
+
+USoundCueDistanceCrossfade::USoundCueDistanceCrossfade(const FObjectInitializer& ObjectInitializer)
+	: Super(ObjectInitializer)
+{
+#if WITH_EDITORONLY_DATA
+// Add editor data initialization here
+	bLooping = true;
+#endif // WITH_EDITORONLY_DATA
+}
+
+#if WITH_EDITOR
+#define LOCTEXT_NAMESPACE "SoundCueDistanceCrossfade"
+
+void USoundCueDistanceCrossfade::OnRebuildGraph(USoundCue& SoundCue) const
+{
+	const UAudioSettings* AudioSettings = GetDefault<UAudioSettings>();
+	if (!AudioSettings)
+	{
+		return;
+	}
+
+	USoundNodeModulator& ModulatorNode = USoundCueTemplate::ConstructSoundNodeRoot<USoundNodeModulator>(SoundCue);
+	ModulatorNode.PitchMin  = 0.95f;
+	ModulatorNode.PitchMax  = 1.05f;
+
+	// Roughly -3dB to 0 dB
+	ModulatorNode.VolumeMin = 0.75f;
+	ModulatorNode.VolumeMax = 1.00f;
+
+	int32 Column = 1; // 1 as root is '0', which was added above.
+	int32 LeafIndex = 0;
+
+	// Building the next two generations in the following nested loops,
+	// so begin tracking both totals overall here.
+	int32 WavePlayerRow = 0;
+
+	const int32 Row = 0;
+	const int32 InputPinIndex = 0;
+	USoundNodeDistanceCrossFade& CrossfadeNode = USoundCueTemplate::ConstructSoundNodeChild<USoundNodeDistanceCrossFade>(
+		SoundCue,
+		&ModulatorNode,
+		Column,
+		Row,
+		InputPinIndex);
+
+	// Create wave players, update crossfade input, and connect to CrossfadeNode
+	int i = 0;
+	for (const FSoundCueCrossfadeInfo& CrossfadeInfo : Variations)
+	{
+		if (CrossfadeInfo.Sound)
+		{
+			const int32 NewIndex = CrossfadeNode.ChildNodes.Num();
+
+			USoundNodeWavePlayer* WavePlayerNode = &USoundCueTemplate::ConstructSoundNodeChild<USoundNodeWavePlayer>(
+				SoundCue,
+				&CrossfadeNode,
+				Column + 1,
+				WavePlayerRow,
+				i++);
+
+			WavePlayerNode->SetSoundWave(CrossfadeInfo.Sound);
+			WavePlayerNode->bLooping = bLooping;
+			CrossfadeNode.CrossFadeInput[NewIndex] = CrossfadeInfo.DistanceInfo;
+			WavePlayerRow++;
+		}
+	}
+}
+
+TSet<FName>& USoundCueDistanceCrossfade::GetWhitelistedCategories()
+{
+	static TSet<FName> Categories;
+	if (Categories.Num() == 0)
+	{
+		// Add exposed categories here that you'd like to show
+		// up in the editor (should match category names properties
+		// are assigned to in the header, otherwise the properties
+		// will be hidden by default).  If you would like to expose
+		// default property categories inherited by USoundCue,
+		// add them here as well.
+		Categories.Add(FName(TEXT("Wave Parameters")));
+	}
+
+	return Categories;
+}
+
+TSharedRef<IDetailCustomization> FSoundCueDistanceCrossfadeDetailCustomization::MakeInstance()
+{
+	return MakeShareable(new FSoundCueDistanceCrossfadeDetailCustomization);
+}
+
+void FSoundCueDistanceCrossfadeDetailCustomization::CustomizeDetails(IDetailLayoutBuilder& DetailLayout)
+{
+	TArray<FName> CategoryNames;
+	DetailLayout.GetCategoryNames(CategoryNames);
+
+	const TSet<FName>& Whitelist = USoundCueDistanceCrossfade::GetWhitelistedCategories();
+	for (FName CategoryName : CategoryNames)
+	{
+		if (!Whitelist.Contains(CategoryName))
+		{
+			DetailLayout.HideCategory(CategoryName);
+		}
+	}
+}
+
+void FSoundCueDistanceCrossfadeDetailCustomization::Register(FPropertyEditorModule& PropertyModule)
+{
+	PropertyModule.RegisterCustomClassLayout("SoundCueDistanceCrossfade", FOnGetDetailCustomizationInstance::CreateStatic(&FSoundCueDistanceCrossfadeDetailCustomization::MakeInstance));
+}
+#undef LOCTEXT_NAMESPACE
+#endif // WITH_EDITOR
diff --git a/Engine/Plugins/Runtime/SoundCueTemplates/Source/SoundCueTemplates/Private/SoundCueTemplateSettings.cpp b/Engine/Plugins/Runtime/SoundCueTemplates/Source/SoundCueTemplates/Private/SoundCueTemplateSettings.cpp
new file mode 100644
--- /dev/null
+++ b/Engine/Plugins/Runtime/SoundCueTemplates/Source/SoundCueTemplates/Private/SoundCueTemplateSettings.cpp
@@ -0,0 +1,177 @@
+// Copyright 1998-2019 Epic Games, Inc. All Rights Reserved.
+#include "SoundCueTemplateSettings.h"
+
+#include "Sound/AudioSettings.h"
+#include "UObject/UObjectHash.h"
+#include "UObject/UObjectIterator.h"
+#include "SoundCueContainer.h"
+#include "SoundCueTemplatesModule.h"
+
+#if WITH_EDITOR
+#include "Toolkits/AssetEditorManager.h"
+#include "Framework/Notifications/NotificationManager.h"
+#include "Widgets/Input/SButton.h"
+#include "Widgets/Notifications/SNotificationList.h"
+#endif // WITH_EDITOR
+
+#define LOCTEXT_NAMESPACE "SoundCueTemplateSettings"
+
+#if WITH_EDITOR
+namespace
+{
+	void PostNotification(const FText& NotificationText)
+	{
+		const UAudioSettings* AudioSettings = GetDefault<UAudioSettings>();
+		check(AudioSettings);
+
+		FNotificationInfo Info(NotificationText);
+		Info.bFireAndForget = true;
+		Info.ExpireDuration = 2.0f;
+		Info.bUseThrobber = true;
+		FSlateNotificationManager::Get().AddNotification(Info);
+	}
+
+	void RebuildNodeGraphs(const ESoundContainerType* ContainerType)
+	{
+		for (TObjectIterator<USoundCueContainer> It; It; ++It)
+		{
+			if (USoundCueContainer* Container = *It)
+			{
+				if (!ContainerType || *ContainerType == Container->ContainerType)
+				{
+					Container->RebuildGraph(*Container);
+				}
+			}
+		}
+	}
+} // namespace <>
+
+
+void FSoundCueTemplateQualitySettingsNotifier::PostQualitySettingsUpdated() const
+{
+	const FText Text = LOCTEXT("NewSoundCueTemplateQualityLevel",
+		"Main Audio Settings Quality Settings level has been added. SoundCueTemplate Quality Settings may need updates.");
+	PostNotification(Text);
+}
+
+void FSoundCueTemplateQualitySettingsNotifier::OnAudioSettingsChanged() const
+{
+	USoundCueTemplateSettings* TemplateSettings = GetMutableDefault<USoundCueTemplateSettings>();
+	check(TemplateSettings);
+	if (TemplateSettings->RebuildQualityLevels())
+	{
+		PostQualitySettingsUpdated();
+	}
+}
+#endif // WITH_EDITOR
+
+FSoundCueTemplateQualitySettings::FSoundCueTemplateQualitySettings()
+	: MaxConcatenatedVariations(8)
+	, MaxRandomizedVariations(8)
+	, MaxMixVariations(8)
+{
+}
+
+USoundCueTemplateSettings::USoundCueTemplateSettings(const FObjectInitializer& ObjectInitializer)
+: Super(ObjectInitializer)
+{
+}
+
+#if WITH_EDITOR
+int32 USoundCueTemplateSettings::GetQualityLevelSettingsNum() const
+{
+	return QualityLevels.Num();
+}
+
+const FSoundCueTemplateQualitySettings& USoundCueTemplateSettings::GetQualityLevelSettings(int32 Index) const
+{
+	return QualityLevels[Index];
+}
+
+void USoundCueTemplateSettings::PostInitProperties()
+{
+	Super::PostInitProperties();
+
+	UAudioSettings* AudioSettings = GetMutableDefault<UAudioSettings>();
+	check(AudioSettings);
+
+	AudioSettings->OnAudioSettingsChanged().AddRaw(&SettingsNotifier, &FSoundCueTemplateQualitySettingsNotifier::OnAudioSettingsChanged);
+}
+
+void USoundCueTemplateSettings::PostEditChangeChainProperty(FPropertyChangedChainEvent& PropertyChangedEvent)
+{
+	if (!PropertyChangedEvent.Property)
+	{
+		return;
+	}
+
+	if (!RebuildQualityLevels())
+	{
+		return;
+	}
+
+	MarkPackageDirty();
+
+	bool bUpdateTemplates = false;
+	bool bFilterByContainerType = false;
+	ESoundContainerType ContainerType = ESoundContainerType::Randomize;
+
+	if (PropertyChangedEvent.Property->GetFName() == GET_MEMBER_NAME_CHECKED(FSoundCueTemplateQualitySettings, MaxConcatenatedVariations))
+	{
+		bUpdateTemplates = true;
+		bFilterByContainerType = true;
+		ContainerType = ESoundContainerType::Concatenate;
+	}
+	else if (PropertyChangedEvent.Property->GetFName() == GET_MEMBER_NAME_CHECKED(FSoundCueTemplateQualitySettings, MaxRandomizedVariations))
+	{
+		bUpdateTemplates = true;
+		bFilterByContainerType = true;
+		ContainerType = ESoundContainerType::Randomize;
+	}
+	else if (PropertyChangedEvent.Property->GetFName() == GET_MEMBER_NAME_CHECKED(FSoundCueTemplateQualitySettings, MaxMixVariations))
+	{
+		bUpdateTemplates = true;
+		bFilterByContainerType = true;
+		ContainerType = ESoundContainerType::Mix;
+	}
+	else if (PropertyChangedEvent.Property->GetFName() == GET_MEMBER_NAME_CHECKED(FSoundCueTemplateQualitySettings, DisplayName))
+	{
+		bUpdateTemplates = true;
+	}
+
+	if (bUpdateTemplates)
+	{
+		RebuildNodeGraphs(bFilterByContainerType ? &ContainerType : nullptr);
+	}
+}
+
+bool USoundCueTemplateSettings::RebuildQualityLevels()
+{
+	const UAudioSettings* AudioSettings = GetDefault<UAudioSettings>();
+	check(AudioSettings);
+
+	bool bHasChanged = false;
+
+	while (QualityLevels.Num() > AudioSettings->GetQualityLevelSettingsNum())
+	{
+		bHasChanged = true;
+		QualityLevels.RemoveAt(QualityLevels.Num() - 1);
+	}
+
+	while (QualityLevels.Num() < AudioSettings->GetQualityLevelSettingsNum())
+	{
+		bHasChanged = true;
+		QualityLevels.Add(FSoundCueTemplateQualitySettings());
+	}
+
+	for (int i = 0; i < AudioSettings->GetQualityLevelSettingsNum(); ++i)
+	{
+		FText& QualityName = QualityLevels[i].DisplayName;
+		const FText& AudioSettingsQualityName = AudioSettings->GetQualityLevelSettings(i).DisplayName;
+		QualityName = AudioSettingsQualityName;
+	}
+
+	return bHasChanged;
+}
+#endif // WITH_EDITOR
+#undef LOCTEXT_NAMESPACE
diff --git a/Engine/Plugins/Runtime/SoundCueTemplates/Source/SoundCueTemplates/Private/SoundCueTemplatesModule.cpp b/Engine/Plugins/Runtime/SoundCueTemplates/Source/SoundCueTemplates/Private/SoundCueTemplatesModule.cpp
new file mode 100644
--- /dev/null
+++ b/Engine/Plugins/Runtime/SoundCueTemplates/Source/SoundCueTemplates/Private/SoundCueTemplatesModule.cpp
@@ -0,0 +1,40 @@
+// Copyright 1998-2019 Epic Games, Inc. All Rights Reserved.
+#include "SoundCueTemplatesModule.h"
+
+#include "Logging/MessageLog.h"
+#include "SoundCueContainer.h"
+#include "SoundCueDistanceCrossfade.h"
+#include "SoundCueTemplateSettings.h"
+
+#if WITH_EDITOR
+#include "ISettingsModule.h"
+#include "ISettingsSection.h"
+#include "PropertyEditorModule.h"
+#endif // WITH_EDITOR
+
+
+DEFINE_LOG_CATEGORY(SoundCueTemplates);
+
+#define LOCTEXT_NAMESPACE "FSoundCueTemplatesModule"
+void FSoundCueTemplatesModule::StartupModule()
+{
+#if WITH_EDITOR
+	FPropertyEditorModule& PropertyModule = FModuleManager::LoadModuleChecked<FPropertyEditorModule>("PropertyEditor");
+
+	FSoundCueContainerDetailCustomization::Register(PropertyModule);
+	FSoundCueDistanceCrossfadeDetailCustomization::Register(PropertyModule);
+#endif //WITH_EDITOR
+}
+
+void FSoundCueTemplatesModule::ShutdownModule()
+{
+#if WITH_EDITOR
+	if (ISettingsModule* SettingsModule = FModuleManager::GetModulePtr<ISettingsModule>("Settings"))
+	{
+		SettingsModule->UnregisterSettings("Project", "Plugins", "SoundCueTemplates");
+	}
+#endif //WITH_EDITOR
+}
+#undef LOCTEXT_NAMESPACE
+
+IMPLEMENT_MODULE(FSoundCueTemplatesModule, SoundCueTemplates)
\ No newline at end of file
diff --git a/Engine/Plugins/Runtime/SoundCueTemplates/Source/SoundCueTemplates/Public/SoundCueContainer.h b/Engine/Plugins/Runtime/SoundCueTemplates/Source/SoundCueTemplates/Public/SoundCueContainer.h
new file mode 100644
--- /dev/null
+++ b/Engine/Plugins/Runtime/SoundCueTemplates/Source/SoundCueTemplates/Public/SoundCueContainer.h
@@ -0,0 +1,93 @@
+// Copyright 1998-2019 Epic Games, Inc. All Rights Reserved.
+
+#pragma once
+
+#include "CoreMinimal.h"
+
+
+#include "Misc/Attribute.h"
+#include "Sound/SoundCueTemplate.h"
+#include "SoundCueTemplateSettings.h"
+#include "UObject/ObjectMacros.h"
+
+#if WITH_EDITOR
+#include "IDetailCustomization.h"
+#include "Layout/Visibility.h"
+#include "PropertyHandle.h"
+#endif // WITH_EDITOR
+
+#include "SoundCueContainer.generated.h"
+
+
+// ========================================================================
+// USoundCueContainer
+// Sound Cue template class which implements USoundCueTemplate.
+//
+// Simple example showing how to expose or hide template
+// parameters in the editor such as the looping and soundwave
+// fields of a USoundNodeWavePlayer.
+//
+// In order for proper data hiding to occur for inherited properties,
+// Customization Detail's 'Register' must be called in during initialization
+// (eg. in module's StartupModule()) like so:
+// #include "SoundCueContainer.h"
+// ...
+// FPropertyEditorModule& PropertyModule = FModuleManager::LoadModuleChecked<FPropertyEditorModule>("PropertyEditor");
+// FSoundCueContainerDetailCustomization::Register(PropertyModule);
+// ========================================================================
+
+
+UENUM()
+enum class ESoundContainerType : uint8
+{
+	Concatenate,
+	Randomize,
+	Mix
+};
+
+UCLASS(hidecategories = object, BlueprintType)
+class SOUNDCUETEMPLATES_API USoundCueContainer : public USoundCueTemplate
+{
+	GENERATED_UCLASS_BODY()
+
+#if WITH_EDITORONLY_DATA
+public:
+	UPROPERTY(EditAnywhere, Category = "Variation")
+	ESoundContainerType ContainerType;
+
+	UPROPERTY(EditAnywhere, Category = "Variation")
+	bool bLooping;
+
+	UPROPERTY(EditAnywhere, Category = "Variation")
+	TSet<USoundWave*> Variations;
+
+	UPROPERTY(EditAnywhere, Category = "Modulation")
+	FVector2D PitchModulation;
+
+	UPROPERTY(EditAnywhere, Category = "Modulation")
+	FVector2D VolumeModulation;
+#endif // WITH_EDITORONLYDATA
+
+#if WITH_EDITOR
+	/** Whitelisted categories, which allows hiding default
+	  * categories from the editor on the inherited SoundCue
+	  * as desired.
+	  */
+	static TSet<FName>& GetWhitelistedCategories();
+	virtual void OnRebuildGraph(USoundCue& SoundCue) const override;
+
+private:
+	int32 GetMaxVariations(const FSoundCueTemplateQualitySettings& QualitySettings) const;
+#endif // WITH_EDITOR
+};
+
+#if WITH_EDITOR
+class FSoundCueContainerDetailCustomization : public IDetailCustomization
+{
+public:
+	static TSharedRef<IDetailCustomization> MakeInstance();
+	static void Register(FPropertyEditorModule& PropertyModule);
+	static void Unregister(FPropertyEditorModule& PropertyModule);
+	virtual void CustomizeDetails(IDetailLayoutBuilder& DetailLayout) override;
+};
+#endif // WITH_EDITOR
diff --git a/Engine/Plugins/Runtime/SoundCueTemplates/Source/SoundCueTemplates/Public/SoundCueDistanceCrossfade.h b/Engine/Plugins/Runtime/SoundCueTemplates/Source/SoundCueTemplates/Public/SoundCueDistanceCrossfade.h
new file mode 100644
--- /dev/null
+++ b/Engine/Plugins/Runtime/SoundCueTemplates/Source/SoundCueTemplates/Public/SoundCueDistanceCrossfade.h
@@ -0,0 +1,77 @@
+// Copyright 1998-2019 Epic Games, Inc. All Rights Reserved.
+
+#pragma once
+
+#include "CoreMinimal.h"
+
+#include "Sound/SoundCueTemplate.h"
+#include "Sound/SoundNodeDistanceCrossFade.h"
+
+#include "Misc/Attribute.h"
+#include "UObject/ObjectMacros.h"
+
+#if WITH_EDITOR
+#include "IDetailCustomization.h"
+#include "Layout/Visibility.h"
+#include "PropertyHandle.h"
+#endif // WITH_EDITOR
+
+#include "SoundCueDistanceCrossfade.generated.h"
+
+// ========================================================================
+// USoundCueDistanceCrossfade
+// Sound Cue template class which implements USoundCueTemplate.
+//
+// Simple near/distant mix design for provided, spatialized sounds.
+// ========================================================================
+
+
+USTRUCT(BlueprintType)
+struct FSoundCueCrossfadeInfo
+{
+	GENERATED_USTRUCT_BODY()
+
+	UPROPERTY(EditAnywhere, Category = "Crossfade Info")
+	FDistanceDatum DistanceInfo;
+
+	UPROPERTY(EditAnywhere, Category = "Crossfade Info")
+	USoundWave* Sound;
+};
+
+UCLASS(hidecategories = object, BlueprintType)
+class SOUNDCUETEMPLATES_API USoundCueDistanceCrossfade : public USoundCueTemplate
+{
+	GENERATED_UCLASS_BODY()
+
+#if WITH_EDITORONLY_DATA
+public:
+	UPROPERTY(EditAnywhere, Category = "Wave Parameters")
+	bool bLooping;
+
+	UPROPERTY(EditAnywhere, Category = "Wave Parameters")
+	TArray<FSoundCueCrossfadeInfo> Variations;
+#endif // WITH_EDITORONLY_DATA
+
+#if WITH_EDITOR
+	/** Whitelisted categories, which allows hiding default
+	  * categories from the editor on the inherited SoundCue
+	  * as desired.
+	  */
+	static TSet<FName>& GetWhitelistedCategories();
+
+	/** Overload with logic to rebuild the node graph when
+	  * user-facing properties are changed
+	  */
+	virtual void OnRebuildGraph(USoundCue& SoundCue) const override;
+#endif // WITH_EDITOR
+};
+
+#if WITH_EDITOR
+class FSoundCueDistanceCrossfadeDetailCustomization : public IDetailCustomization
+{
+public:
+	static TSharedRef<IDetailCustomization> MakeInstance();
+	static void Register(FPropertyEditorModule& PropertyModule);
+	virtual void CustomizeDetails(IDetailLayoutBuilder& DetailLayout) override;
+};
+#endif // WITH_EDITOR
diff --git a/Engine/Plugins/Runtime/SoundCueTemplates/Source/SoundCueTemplates/Public/SoundCueTemplateSettings.h b/Engine/Plugins/Runtime/SoundCueTemplates/Source/SoundCueTemplates/Public/SoundCueTemplateSettings.h
new file mode 100644
--- /dev/null
+++ b/Engine/Plugins/Runtime/SoundCueTemplates/Source/SoundCueTemplates/Public/SoundCueTemplateSettings.h
@@ -0,0 +1,63 @@
+// Copyright 1998-2019 Epic Games, Inc. All Rights Reserved.
+#pragma once
+
+#include "CoreMinimal.h"
+#include "UObject/ObjectMacros.h"
+#include "UObject/SoftObjectPath.h"
+#include "Engine/DeveloperSettings.h"
+
+#include "SoundCueTemplateSettings.generated.h"
+
+
+struct FSoundCueTemplateQualitySettingsNotifier
+{
+	void PostQualitySettingsUpdated() const;
+	void OnAudioSettingsChanged() const;
+};
+
+USTRUCT()
+struct FSoundCueTemplateQualitySettings
+{
+	GENERATED_BODY()
+
+	UPROPERTY(VisibleAnywhere, Category = "Quality")
+	FText DisplayName;
+
+	// The max number of variations to include for the given quality in a SoundCueContainer set to 'Concatenate'.
+	UPROPERTY(EditAnywhere, Category = "Quality", meta = (ClampMin = "1"))
+	int32 MaxConcatenatedVariations;
+
+	// The max number of variations to include for the given quality in a SoundCueContainer set to 'Randomized'.
+	UPROPERTY(EditAnywhere, Category = "Quality", meta = (ClampMin = "1"))
+	int32 MaxRandomizedVariations;
+
+	// The max number of variations to include for the given quality in a SoundCueContainer set to 'Mix'.
+	UPROPERTY(EditAnywhere, Category = "Quality", meta = (ClampMin = "1"))
+	int32 MaxMixVariations;
+
+	FSoundCueTemplateQualitySettings();
+};
+
+UCLASS(config=Game, defaultconfig, meta=(DisplayName="Sound Cue Templates"))
+class USoundCueTemplateSettings : public UDeveloperSettings
+{
+	GENERATED_UCLASS_BODY()
+
+#if WITH_EDITORONLY_DATA
+	UPROPERTY(config, EditAnywhere, Category = "Quality")
+	TArray<FSoundCueTemplateQualitySettings> QualityLevels;
+#endif // WITH_EDITORONLY_DATA
+
+#if WITH_EDITOR
+	virtual void PostEditChangeChainProperty( struct FPropertyChangedChainEvent& PropertyChangedEvent) override;
+	virtual void PostInitProperties() override;
+
+	FSoundCueTemplateQualitySettingsNotifier SettingsNotifier;
+public:
+	bool RebuildQualityLevels();
+
+	// Get the quality level settings at the provided level index
+	const FSoundCueTemplateQualitySettings& GetQualityLevelSettings(int32 Index) const;
+	int32 GetQualityLevelSettingsNum() const;
+#endif // WITH_EDITOR
+};
diff --git a/Engine/Plugins/Runtime/SoundCueTemplates/Source/SoundCueTemplates/Public/SoundCueTemplatesModule.h b/Engine/Plugins/Runtime/SoundCueTemplates/Source/SoundCueTemplates/Public/SoundCueTemplatesModule.h
new file mode 100644
--- /dev/null
+++ b/Engine/Plugins/Runtime/SoundCueTemplates/Source/SoundCueTemplates/Public/SoundCueTemplatesModule.h
@@ -0,0 +1,17 @@
+// Copyright 1998-2019 Epic Games, Inc. All Rights Reserved.
+
+#pragma once
+
+#include "CoreMinimal.h"
+#include "Modules/ModuleManager.h"
+
+DECLARE_LOG_CATEGORY_EXTERN(SoundCueTemplates, Log, All);
+
+class FSoundCueTemplatesModule : public IModuleInterface
+{
+public:
+
+	/** IModuleInterface implementation */
+	virtual void StartupModule() override;
+	virtual void ShutdownModule() override;
+};
diff --git a/Engine/Plugins/Runtime/SoundCueTemplates/Source/SoundCueTemplates/SoundCueTemplates.Build.cs b/Engine/Plugins/Runtime/SoundCueTemplates/Source/SoundCueTemplates/SoundCueTemplates.Build.cs
new file mode 100644
--- /dev/null
+++ b/Engine/Plugins/Runtime/SoundCueTemplates/Source/SoundCueTemplates/SoundCueTemplates.Build.cs
@@ -0,0 +1,29 @@
+// Copyright 1998-2019 Epic Games, Inc. All Rights Reserved.
+
+using UnrealBuildTool;
+
+public class SoundCueTemplates : ModuleRules
+{
+    public SoundCueTemplates(ReadOnlyTargetRules Target) : base(Target)
+    {
+        OptimizeCode = CodeOptimization.Never;
+        PCHUsage = ModuleRules.PCHUsageMode.UseExplicitOrSharedPCHs;
+
+        PublicDependencyModuleNames.AddRange(
+            new string[]
+            {
+                "Core"
+            }
+        );
+
+        PrivateDependencyModuleNames.AddRange(
+            new string[]
+            {
+                "CoreUObject",
+                "Engine",
+                "Slate",
+                "SlateCore"
+            }
+        );
+    }
+}
diff --git a/Engine/Source/Editor/AudioEditor/Classes/Factories/SoundCueTemplateFactory.h b/Engine/Source/Editor/AudioEditor/Classes/Factories/SoundCueTemplateFactory.h
new file mode 100644
--- /dev/null
+++ b/Engine/Source/Editor/AudioEditor/Classes/Factories/SoundCueTemplateFactory.h
@@ -0,0 +1,37 @@
+// Copyright 1998-2019 Epic Games, Inc. All Rights Reserved.
+#pragma once
+
+#include "CoreMinimal.h"
+#include "UObject/ObjectMacros.h"
+#include "Factories/Factory.h"
+
+#include "SoundCueTemplateFactory.generated.h"
+
+// Forward Declarations
+class USoundCue;
+class USoundCueTemplate;
+class USoundWave;
+
+UCLASS(hidecategories = Object, MinimalAPI)
+class USoundCueTemplateCopyFactory : public UFactory
+{
+	GENERATED_UCLASS_BODY()
+
+	//~ Begin UFactory Interface
+	virtual UObject* FactoryCreateNew(UClass* Class, UObject* InParent, FName Name, EObjectFlags Flags, UObject* Context, FFeedbackContext* Warn) override;
+	//~ Begin UFactory Interface
+
+public:
+	UPROPERTY()
+	TWeakObjectPtr<USoundCueTemplate> SoundCueTemplate;
+};
+
+UCLASS(hidecategories = Object, MinimalAPI)
+class USoundCueTemplateFactory : public UFactory
+{
+	GENERATED_UCLASS_BODY()
+
+	//~ Begin UFactory Interface
+	virtual UObject* FactoryCreateNew(UClass* Class, UObject* InParent, FName Name, EObjectFlags Flags, UObject* Context, FFeedbackContext* Warn) override;
+	//~ Begin UFactory Interface
+};
\ No newline at end of file
diff --git a/Engine/Source/Editor/AudioEditor/Private/AssetTypeActions/AssetTypeActions_SoundCueTemplate.cpp b/Engine/Source/Editor/AudioEditor/Private/AssetTypeActions/AssetTypeActions_SoundCueTemplate.cpp
new file mode 100644
--- /dev/null
+++ b/Engine/Source/Editor/AudioEditor/Private/AssetTypeActions/AssetTypeActions_SoundCueTemplate.cpp
@@ -0,0 +1,60 @@
+// Copyright 1998-2019 Epic Games, Inc. All Rights Reserved.
+#include "AssetTypeActions/AssetTypeActions_SoundCueTemplate.h"
+
+#include "ContentBrowserModule.h"
+#include "EditorStyleSet.h"
+#include "Factories/SoundCueFactoryNew.h"
+#include "Factories/SoundCueTemplateFactory.h"
+#include "Framework/MultiBox/MultiBoxBuilder.h"
+#include "IContentBrowserSingleton.h"
+#include "Sound/SoundCue.h"
+#include "Sound/SoundCueTemplate.h"
+#include "Widgets/Layout/SBox.h"
+#include "Widgets/Input/SButton.h"
+#include "Widgets/Images/SImage.h"
+
+
+#define LOCTEXT_NAMESPACE "AssetTypeActions"
+
+UClass* FAssetTypeActions_SoundCueTemplate::GetSupportedClass() const
+{
+	return USoundCueTemplate::StaticClass();
+}
+
+void FAssetTypeActions_SoundCueTemplate::GetActions(const TArray<UObject*>& InObjects, FMenuBuilder& MenuBuilder)
+{
+	TArray<TWeakObjectPtr<USoundCueTemplate>> Cues = GetTypedWeakObjectPtrs<USoundCueTemplate>(InObjects);
+
+	MenuBuilder.AddMenuEntry(
+		LOCTEXT("SoundCueTemplate_CopyToSoundCue", "Copy To Sound Cue"),
+		LOCTEXT("SoundCueTemplate_CopyToSoundCueTooltip", "Exports a Sound Cue Template to a Sound Cue."),
+		FSlateIcon(FEditorStyle::GetStyleSetName(), "ClassIcon.SoundCue"),
+		FUIAction(
+			FExecuteAction::CreateSP(this, &FAssetTypeActions_SoundCueTemplate::ExecuteCopyToSoundCue, Cues),
+			FCanExecuteAction()
+		)
+	);
+}
+
+void FAssetTypeActions_SoundCueTemplate::ExecuteCopyToSoundCue(TArray<TWeakObjectPtr<USoundCueTemplate>> Objects)
+{
+	for (const TWeakObjectPtr<USoundCueTemplate>& Object : Objects)
+	{
+		if (!Object.IsValid())
+		{
+			continue;
+		}
+
+		FString Name;
+		FString PackagePath;
+		CreateUniqueAssetName(Object->GetOutermost()->GetName(), TEXT(""), PackagePath, Name);
+
+		if (USoundCueTemplateCopyFactory* Factory = NewObject<USoundCueTemplateCopyFactory>())
+		{
+			Factory->SoundCueTemplate = Object;
+			FContentBrowserModule& ContentBrowserModule = FModuleManager::LoadModuleChecked<FContentBrowserModule>("ContentBrowser");
+			ContentBrowserModule.Get().CreateNewAsset(Name, FPackageName::GetLongPackagePath(PackagePath), USoundCue::StaticClass(), Factory);
+		}
+	}
+}
+#undef LOCTEXT_NAMESPACE
diff --git a/Engine/Source/Editor/AudioEditor/Private/AssetTypeActions/AssetTypeActions_SoundCueTemplate.h b/Engine/Source/Editor/AudioEditor/Private/AssetTypeActions/AssetTypeActions_SoundCueTemplate.h
new file mode 100644
--- /dev/null
+++ b/Engine/Source/Editor/AudioEditor/Private/AssetTypeActions/AssetTypeActions_SoundCueTemplate.h
@@ -0,0 +1,25 @@
+// Copyright 1998-2019 Epic Games, Inc. All Rights Reserved.
+#pragma once
+
+#include "CoreMinimal.h"
+#include "Toolkits/IToolkitHost.h"
+#include "AssetTypeActions/AssetTypeActions_SoundBase.h"
+
+class FMenuBuilder;
+class USoundCueTemplate;
+
+class FAssetTypeActions_SoundCueTemplate : public FAssetTypeActions_SoundBase
+{
+public:
+	// IAssetTypeActions Implementation
+	virtual FText GetName() const override { return NSLOCTEXT("AssetTypeActions", "AssetTypeActions_SoundCueTemplate", "Sound Cue Template"); }
+	virtual FColor GetTypeColor() const override { return FColor(255, 255, 255); }
+	virtual UClass* GetSupportedClass() const override;
+	virtual void GetActions(const TArray<UObject*>& InObjects, FMenuBuilder& MenuBuilder) override;
+	virtual bool CanFilter() override { return true; }
+
+protected:
+	/** Converts the provided SoundCue Template to a fully-modifiable SoundCue */
+	void ExecuteCopyToSoundCue(TArray<TWeakObjectPtr<USoundCueTemplate>> Objects);
+};
+
diff --git a/Engine/Source/Editor/AudioEditor/Private/AudioEditorModule.cpp b/Engine/Source/Editor/AudioEditor/Private/AudioEditorModule.cpp
--- a/Engine/Source/Editor/AudioEditor/Private/AudioEditorModule.cpp
+++ b/Engine/Source/Editor/AudioEditor/Private/AudioEditorModule.cpp
@@ -11,6 +11,7 @@
 #include "SoundCueGraphNodeFactory.h"
 #include "Factories/ReimportSoundSurroundFactory.h"
 #include "AssetToolsModule.h"
+#include "PropertyEditorModule.h"
 #include "SoundClassEditor.h"
 #include "Sound/SoundCue.h"
 #include "Sound/SoundWave.h"
@@ -26,6 +27,7 @@
 #include "AssetTypeActions/AssetTypeActions_SoundBase.h"
 #include "AssetTypeActions/AssetTypeActions_SoundClass.h"
 #include "AssetTypeActions/AssetTypeActions_SoundCue.h"
+#include "AssetTypeActions/AssetTypeActions_SoundCueTemplate.h"
 #include "AssetTypeActions/AssetTypeActions_SoundMix.h"
 #include "AssetTypeActions/AssetTypeActions_SoundWave.h"
 #include "AssetTypeActions/AssetTypeActions_ReverbEffect.h"
@@ -55,7 +57,7 @@ static const FVector2D Icon64 = FVector2D(64.0f, 64.0f);
 // ICON_NAME - base-name of the icon to use. Not necessarily based off class name
 #define SET_AUDIO_ICON(CLASS_NAME, ICON_NAME) \
 		AudioStyleSet->Set( *FString::Printf(TEXT("ClassIcon.%s"), TEXT(#CLASS_NAME)), new FSlateImageBrush(FPaths::EngineContentDir() / FString::Printf(TEXT("Editor/Slate/Icons/AssetIcons/%s_16x.png"), TEXT(#ICON_NAME)), Icon16)); \
-		AudioStyleSet->Set( *FString::Printf(TEXT("ClassThumbnail.%s"), TEXT(#CLASS_NAME)), new FSlateImageBrush(FPaths::EngineContentDir() / FString::Printf(TEXT("Editor/Slate/Icons/AssetIcons/%s_64x.png"), TEXT(#ICON_NAME)), Icon64)); 
+		AudioStyleSet->Set( *FString::Printf(TEXT("ClassThumbnail.%s"), TEXT(#CLASS_NAME)), new FSlateImageBrush(FPaths::EngineContentDir() / FString::Printf(TEXT("Editor/Slate/Icons/AssetIcons/%s_64x.png"), TEXT(#ICON_NAME)), Icon64));
 
 // Simpler version of SET_AUDIO_ICON, assumes same name of icon png and class name
 #define SET_AUDIO_ICON_SIMPLE(CLASS_NAME) SET_AUDIO_ICON(CLASS_NAME, CLASS_NAME)
@@ -90,7 +92,6 @@ public:
 		UReimportSoundSurroundFactory::StaticClass();
 
 		SetupIcons();
-
 #if WITH_SNDFILE_IO
 		if (!Audio::InitSoundFileIOManager())
 		{
@@ -127,6 +128,7 @@ public:
 		AssetTools.RegisterAssetTypeActions(MakeShareable(new FAssetTypeActions_SoundBase));
 		AssetTools.RegisterAssetTypeActions(MakeShareable(new FAssetTypeActions_SoundClass));
 		AssetTools.RegisterAssetTypeActions(MakeShareable(new FAssetTypeActions_SoundCue));
+		AssetTools.RegisterAssetTypeActions(MakeShareable(new FAssetTypeActions_SoundCueTemplate));
 		AssetTools.RegisterAssetTypeActions(MakeShareable(new FAssetTypeActions_SoundMix));
 		AssetTools.RegisterAssetTypeActions(MakeShareable(new FAssetTypeActions_SoundWave));
 		AssetTools.RegisterAssetTypeActions(MakeShareable(new FAssetTypeActions_ReverbEffect));
diff --git a/Engine/Source/Editor/AudioEditor/Private/Factories/SoundCueTemplateFactory.cpp b/Engine/Source/Editor/AudioEditor/Private/Factories/SoundCueTemplateFactory.cpp
new file mode 100644
--- /dev/null
+++ b/Engine/Source/Editor/AudioEditor/Private/Factories/SoundCueTemplateFactory.cpp
@@ -0,0 +1,73 @@
+// Copyright 1998-2019 Epic Games, Inc. All Rights Reserved.
+#include "Factories/SoundCueTemplateFactory.h"
+
+#include "ClassViewerModule.h"
+#include "EdGraph/EdGraphPin.h"
+#include "Engine/Engine.h"
+#include "Kismet2/SClassPickerDialog.h"
+#include "SoundCueGraph/SoundCueGraphNode.h"
+#include "Sound/SoundCue.h"
+#include "Sound/SoundCueTemplate.h"
+#include "SoundFactoryUtility.h"
+#include "UObject/Class.h"
+
+
+USoundCueTemplateCopyFactory::USoundCueTemplateCopyFactory(const FObjectInitializer& ObjectInitializer)
+	: Super(ObjectInitializer)
+{
+	SupportedClass = USoundCue::StaticClass();
+
+	bCreateNew = false;
+	bEditorImport = false;
+	bEditAfterNew = true;
+}
+
+UObject* USoundCueTemplateCopyFactory::FactoryCreateNew(UClass* Class, UObject* InParent, FName Name, EObjectFlags Flags, UObject* Context, FFeedbackContext* Warn)
+{
+	if (SoundCueTemplate.IsValid())
+	{
+		if (USoundCue* NewCue = NewObject<USoundCue>(InParent, Name, Flags))
+		{
+			SoundCueTemplate.Get()->RebuildGraph(*NewCue);
+			return NewCue;
+		}
+	}
+
+	return nullptr;
+}
+
+USoundCueTemplateFactory::USoundCueTemplateFactory(const FObjectInitializer& ObjectInitializer)
+	: Super(ObjectInitializer)
+{
+	SupportedClass = USoundCueTemplate::StaticClass();
+
+	bCreateNew = true;
+	bEditorImport = false;
+	bEditAfterNew = true;
+}
+
+UObject* USoundCueTemplateFactory::FactoryCreateNew(UClass* Class, UObject* InParent, FName Name, EObjectFlags Flags, UObject* Context, FFeedbackContext* Warn)
+{
+	FClassViewerModule& ClassViewerModule = FModuleManager::LoadModuleChecked<FClassViewerModule>("ClassViewer");
+
+	TSharedPtr<FAssetClassParentFilter> Filter = MakeShareable(new FAssetClassParentFilter);
+	Filter->DisallowedClassFlags = CLASS_Abstract | CLASS_Deprecated | CLASS_NewerVersionExists | CLASS_HideDropDown;
+	Filter->AllowedChildrenOfClasses.Add(USoundCueTemplate::StaticClass());
+
+	FClassViewerInitializationOptions Options;
+	Options.Mode = EClassViewerMode::ClassPicker;
+	Options.ClassFilter = Filter;
+
+	const FText TitleText = FText::FromString(TEXT("Pick Sound Cue Template"));
+	UClass* ChosenClass = nullptr;
+	if (SClassPickerDialog::PickClass(TitleText, Options, ChosenClass, USoundCueTemplate::StaticClass()))
+	{
+		if (USoundCueTemplate* SoundCueTemplate = NewObject<USoundCueTemplate>(InParent, ChosenClass, Name, Flags))
+		{
+			SoundCueTemplate->RebuildGraph(*SoundCueTemplate);
+			return SoundCueTemplate;
+		}
+	}
+
+	return nullptr;
+}
diff --git a/Engine/Source/Editor/GameProjectGeneration/Private/GameProjectUtils.cpp b/Engine/Source/Editor/GameProjectGeneration/Private/GameProjectUtils.cpp
--- a/Engine/Source/Editor/GameProjectGeneration/Private/GameProjectUtils.cpp
+++ b/Engine/Source/Editor/GameProjectGeneration/Private/GameProjectUtils.cpp
@@ -69,6 +69,7 @@
 #include "IAudioExtensionPlugin.h"
 #include "AudioPluginUtilities.h"
 #include "Sound/AudioSettings.h"
+#include "Sound/SoundCueTemplate.h"
 #include "Sound/SoundEffectSubmix.h"
 #include "Sound/SoundEffectSource.h"
 #include "Components/SynthComponent.h"
@@ -451,6 +452,10 @@ FString FNewClassInfo::GetHeaderTemplateFilename() const
 				{
 					return TEXT("SynthComponentClass.h.template");
 				}
+				else if (BaseClass == USoundCueTemplate::StaticClass())
+				{
+					return TEXT("SoundCueTemplateClass.h.template");
+				}
 			}
 			// Some other non-actor, non-component UObject class
 			return TEXT( "UObjectClass.h.template" );
@@ -509,6 +514,10 @@ FString FNewClassInfo::GetSourceTemplateFilename() const
 				{
 					return TEXT("SynthComponentClass.cpp.template");
 				}
+				else if (BaseClass == USoundCueTemplate::StaticClass())
+				{
+					return TEXT("SoundCueTemplateClass.cpp.template");
+				}
 			}
 			// Some other non-actor, non-component UObject class
 			return TEXT( "UObjectClass.cpp.template" );
diff --git a/Engine/Source/Runtime/Engine/Classes/Sound/AudioSettings.h b/Engine/Source/Runtime/Engine/Classes/Sound/AudioSettings.h
--- a/Engine/Source/Runtime/Engine/Classes/Sound/AudioSettings.h
+++ b/Engine/Source/Runtime/Engine/Classes/Sound/AudioSettings.h
@@ -8,6 +8,8 @@
 #include "Engine/DeveloperSettings.h"
 #include "AudioSettings.generated.h"
 
+struct FPropertyChangedChainEvent;
+
 struct ENGINE_API FAudioPlatformSettings
 {
 	/** Sample rate to use on the platform for the mixing engine. Higher sample rates will incur more CPU cost. */
@@ -100,8 +102,11 @@ class ENGINE_API UAudioSettings : public UDeveloperSettings
 
 #if WITH_EDITOR
 	virtual void PreEditChange(UProperty* PropertyAboutToChange) override;
-	virtual void PostEditChangeChainProperty( struct FPropertyChangedChainEvent& PropertyChangedEvent) override;
-#endif
+	virtual void PostEditChangeChainProperty(FPropertyChangedChainEvent& PropertyChangedEvent) override;
+
+	/** Event to listen for when settings reflected properties are changed. */
+	DECLARE_EVENT(UAudioSettings, FAudioSettingsChanged)
+#endif // WITH_EDITOR
 
 	virtual void Serialize(FArchive& Ar) override;
 
@@ -202,8 +207,17 @@ class ENGINE_API UAudioSettings : public UDeveloperSettings
 	UPROPERTY(config, EditAnywhere, Category="Dialogue")
 	FString DialogueFilenameFormat;
 
+#if WITH_EDITOR
+	FAudioSettingsChanged AudioSettingsChanged;
+#endif // WITH_EDITOR
+
+public:
+	// Get the quality level settings at the provided level index
 	const FAudioQualitySettings& GetQualityLevelSettings(int32 QualityLevel) const;
 
+	// Get the total number of quality level settings
+	int32 GetQualityLevelSettingsNum() const;
+
 	// Sets whether audio mixer is enabled. Set once an audio mixer platform module is loaded.
 	void SetAudioMixerEnabled(const bool bInAudioMixerEnabled);
 
@@ -213,11 +227,15 @@ class ENGINE_API UAudioSettings : public UDeveloperSettings
 	/** Returns the highest value for MaxChannels among all quality levels */
 	int32 GetHighestMaxChannels() const;
 
-private:
+#if WITH_EDITOR
+	/** Returns event to be bound to if caller wants to know when audio settings are modified */
+	FAudioSettingsChanged& OnAudioSettingsChanged() { return AudioSettingsChanged; }
+#endif // WITH_EDITOR
 
+private:
 #if WITH_EDITOR
 	TArray<FAudioQualitySettings> CachedQualityLevels;
-#endif
+#endif // WITH_EDITOR
 
 	void AddDefaultSettings();
 
diff --git a/Engine/Source/Runtime/Engine/Classes/Sound/SoundCue.h b/Engine/Source/Runtime/Engine/Classes/Sound/SoundCue.h
--- a/Engine/Source/Runtime/Engine/Classes/Sound/SoundCue.h
+++ b/Engine/Source/Runtime/Engine/Classes/Sound/SoundCue.h
@@ -1,5 +1,4 @@
 // Copyright 1998-2019 Epic Games, Inc. All Rights Reserved.
-
 #pragma once
 
 #include "CoreMinimal.h"
@@ -73,8 +72,8 @@ public:
 /**
  * The behavior of audio playback is defined within Sound Cues.
  */
-UCLASS(hidecategories=object, MinimalAPI, BlueprintType)
-class USoundCue : public USoundBase
+UCLASS(hidecategories=object, BlueprintType)
+class ENGINE_API USoundCue : public USoundBase
 {
 	GENERATED_UCLASS_BODY()
 
@@ -87,7 +86,7 @@ class USoundCue : public USoundBase
 	uint32 bExcludeFromRandomNodeBranchCulling:1;
 
 	UPROPERTY()
-	class USoundNode* FirstNode;
+	USoundNode* FirstNode;
 
 	/* Volume multiplier for the Sound Cue */
 	UPROPERTY(EditAnywhere, Category=Sound, AssetRegistrySearchable)
@@ -180,7 +179,7 @@ public:
 	/**
 	 * Recursively finds all Nodes in the Tree
 	 */
-	ENGINE_API void RecursiveFindAllNodes( USoundNode* Node, TArray<class USoundNode*>& OutNodes );
+	void RecursiveFindAllNodes( USoundNode* Node, TArray<USoundNode*>& OutNodes );
 
 	/**
 	 * Recursively finds sound nodes of type T
@@ -227,15 +226,15 @@ public:
 
 
 	/** Find the path through the sound cue to a node identified by its hash */
-	ENGINE_API bool FindPathToNode(const UPTRINT NodeHashToFind, TArray<USoundNode*>& OutPath) const;
+	bool FindPathToNode(const UPTRINT NodeHashToFind, TArray<USoundNode*>& OutPath) const;
 
 	/** Call when the audio quality has been changed */
-	ENGINE_API static void StaticAudioQualityChanged(int32 NewQualityLevel);
+	static void StaticAudioQualityChanged(int32 NewQualityLevel);
 
 	FORCEINLINE static int32 GetCachedQualityLevel() { return CachedQualityLevel; }
 
 	/** Call to cache any values which need to be computed from the sound cue graph. e.g. MaxDistance, Duration, etc. */
-	ENGINE_API void CacheAggregateValues();
+	void CacheAggregateValues();
 
 protected:
 	bool RecursiveFindPathToNode(USoundNode* CurrentNode, const UPTRINT CurrentHash, const UPTRINT NodeHashToFind, TArray<USoundNode*>& OutPath) const;
@@ -254,32 +253,35 @@ public:
 	/**
 	 * Instantiate certain functions to work around a linker issue
 	 */
-	ENGINE_API void RecursiveFindAttenuation( USoundNode* Node, TArray<class USoundNodeAttenuation*> &OutNodes );
+	void RecursiveFindAttenuation( USoundNode* Node, TArray<class USoundNodeAttenuation*> &OutNodes );
 
 #if WITH_EDITOR
 	/** Create the basic sound graph */
-	ENGINE_API void CreateGraph();
+	void CreateGraph();
 
 	/** Clears all nodes from the graph (for old editor's buffer soundcue) */
-	ENGINE_API void ClearGraph();
+	void ClearGraph();
 
 	/** Set up EdGraph parts of a SoundNode */
-	ENGINE_API void SetupSoundNode(class USoundNode* InSoundNode, bool bSelectNewNode = true);
+	void SetupSoundNode(USoundNode* InSoundNode, bool bSelectNewNode = true);
 
 	/** Use the SoundCue's children to link EdGraph Nodes together */
-	ENGINE_API void LinkGraphNodesFromSoundNodes();
+	void LinkGraphNodesFromSoundNodes();
 
 	/** Use the EdGraph representation to compile the SoundCue */
-	ENGINE_API void CompileSoundNodesFromGraphNodes();
+	void CompileSoundNodesFromGraphNodes();
 
 	/** Get the EdGraph of SoundNodes */
-	ENGINE_API class UEdGraph* GetGraph();
+	class UEdGraph* GetGraph();
+
+	/** Resets all graph data and nodes */
+	void ResetGraph();
 
 	/** Sets the sound cue graph editor implementation.* */
-	static ENGINE_API void SetSoundCueAudioEditor(TSharedPtr<ISoundCueAudioEditor> InSoundCueGraphEditor);
+	static void SetSoundCueAudioEditor(TSharedPtr<ISoundCueAudioEditor> InSoundCueGraphEditor);
 
 	/** Gets the sound cue graph editor implementation. */
-	static TSharedPtr<ISoundCueAudioEditor> ENGINE_API GetSoundCueAudioEditor();
+	static TSharedPtr<ISoundCueAudioEditor> GetSoundCueAudioEditor();
 
 private:
 
@@ -287,6 +289,6 @@ private:
 	void RecursivelySetExcludeBranchCulling(USoundNode* CurrentNode);
 
 	/** Ptr to interface to sound cue editor operations. */
-	static ENGINE_API TSharedPtr<ISoundCueAudioEditor> SoundCueAudioEditor;
+	static TSharedPtr<ISoundCueAudioEditor> SoundCueAudioEditor;
 #endif // WITH_EDITOR
 };
diff --git a/Engine/Source/Runtime/Engine/Classes/Sound/SoundCueTemplate.h b/Engine/Source/Runtime/Engine/Classes/Sound/SoundCueTemplate.h
new file mode 100644
--- /dev/null
+++ b/Engine/Source/Runtime/Engine/Classes/Sound/SoundCueTemplate.h
@@ -0,0 +1,175 @@
+// Copyright 1998-2019 Epic Games, Inc. All Rights Reserved.
+#pragma once
+
+#include "CoreMinimal.h"
+
+#include "UObject/ObjectMacros.h"
+#include "Sound/SoundCue.h"
+#include "Sound/SoundNode.h"
+#include "Templates/SubclassOf.h"
+
+#if WITH_EDITOR
+#include "EdGraph/EdGraphPin.h"
+#endif // WITH_EDITOR
+
+#include "SoundCueTemplate.generated.h"
+
+
+// Base Sound Cue Template class, which builds the sound node graph procedurally and hides more complex Sound Cue functionality
+// to streamline implementation defined in child classes.
+UCLASS(abstract)
+class ENGINE_API USoundCueTemplate : public USoundCue
+{
+	GENERATED_UCLASS_BODY()
+
+#if WITH_EDITOR
+	virtual void PostEditChangeProperty(FPropertyChangedEvent& PropertyChangedEvent) override;
+
+public:
+	/** Rebuilds the graph when user-facing properties are changed */
+	void RebuildGraph(USoundCue& SoundCue) const;
+
+protected:
+	/**
+	 * Function to override that uses internal editor-data only properties to re-build node
+	 * graph of provided SoundCue.  Provided SoundCue may be a reference to this SoundCue
+	 * or one to be copied from this template's graph-building schema.
+	 */
+	virtual void OnRebuildGraph(USoundCue& SoundCue) const PURE_VIRTUAL(, );
+
+	/**
+	 * Node width offset between nodes in the SoundCueEditor.
+	 */
+	static float GetNodeWidthOffset()
+	{
+		return 500.0f;
+	}
+
+	/**
+	 * Initial width offset of first node in the SoundCueEditor.
+	 */
+	static float GetInitialWidthOffset()
+	{
+		return 300.0f;
+	}
+
+	/**
+	 * Node height offset between nodes in the SoundCueEditor.
+	 */
+	static float GetNodeHeightOffset()
+	{
+		return 200.0f;
+	}
+
+	/**
+	 * Initial height offset of first node in the SoundCueEditor.
+	 */
+	static float GetInitialHeightOffset()
+	{
+		return 0.0f;
+	}
+
+	/**
+	 * Sets the provided node's position to the location on the template's grid
+	 */
+	static void SetSoundNodePosition(USoundNode& InNode, int32 Column, int32 Row)
+	{
+		InNode.GraphNode->NodePosX -= Column * GetNodeWidthOffset() + GetInitialWidthOffset();
+		InNode.GraphNode->NodePosY += Row * GetNodeHeightOffset() + GetInitialHeightOffset();
+	}
+
+	/**
+	 * Utility function that returns newly constructed sound node that has been added to provided SoundCue's Node Graph.
+	 * SoundCue				- SoundCue to insert the constructed node into.
+	 * InParentNode			- If not null, Node to set as parent of constructed child.
+	 * Column				- (Optional) If positive, sets visual column of node in built graph as displayed in the SoundCueEditor.  Columns are indexed as positive
+	 *						  to the left from the root (index 0).  Can be viewed when using the 'Copy To SoundCue' action of the given template.
+	 * Row					- (Optional) If positive, sets visual row of node in built graph as displayed in the SoundCueEditor.  Columns are indexed as positive
+	 *						  descending from the root (index 0). Can be viewed when using the 'Copy To SoundCue' action of the given template.
+	 * InputPinIndex		- (Optional) If InParentNode set, supplies index of input pin to connect newly constructed ChildNode to.
+	 * InputPinDisplayName	- (Optional) Name to set the targeted input pin to.
+	 */
+	template <class T>
+	static T& ConstructSoundNodeChild(USoundCue& SoundCue, USoundNode* InParentNode, int32 Column, int32 Row, int32 InputPinIndex = -1, const FText* InputPinDisplayName = nullptr)
+	{
+		T* ChildNode = SoundCue.ConstructSoundNode<T>();
+		check(ChildNode);
+
+		SetSoundNodePosition(*ChildNode, Column, Row);
+
+		if (InParentNode)
+		{
+			AddSoundNodeChild(*InParentNode, *ChildNode, InputPinIndex, InputPinDisplayName);
+		}
+
+		return *ChildNode;
+	}
+
+	/**
+	 * Constructs initial root sound node of template. Function will fail if SoundCue already contains a root/first node.
+	 * SoundCue	- Parent of initial node to construct.
+	 */
+	template <class T>
+	static T& ConstructSoundNodeRoot(USoundCue& SoundCue)
+	{
+		T* RootNode = SoundCue.ConstructSoundNode<T>();
+		check(RootNode);
+		check(!SoundCue.FirstNode);
+
+		SoundCue.FirstNode = RootNode;
+		RootNode->GraphNode->NodePosX = -1 * GetInitialWidthOffset();
+		RootNode->GraphNode->NodePosY = GetInitialHeightOffset();
+
+		return *RootNode;
+	}
+
+	/**
+	 * Utility function that adds a provided child node to a parent node's list of children.
+	 * InParentNode			- Node to set as parent of provided child.
+	 * InChildNode			- Node to set as child of provided parent.
+	 * InputPinIndex		- Index of input pin to connect newly constructed ChildNode to.
+	 * InputPinDisplayName	- (Optional) Name to set the targeted input pin to.
+	 */
+	static void AddSoundNodeChild(USoundNode& InParentNode, USoundNode& InChildNode, int32 InputPinIndex, const FText* InputPinDisplayName = nullptr)
+	{
+		check(InputPinIndex >= 0 && InputPinIndex <= InParentNode.GetMaxChildNodes());
+
+		InParentNode.InsertChildNode(InParentNode.ChildNodes.Num());
+
+		UEdGraphNode* OutNodeGraph = InParentNode.GetGraphNode();
+		check(OutNodeGraph);
+		UEdGraphPin* OutputPin = nullptr;
+		for (UEdGraphPin* ActivePin : InChildNode.GetGraphNode()->Pins)
+		{
+			if (ActivePin->Direction == EGPD_Output)
+			{
+				OutputPin = ActivePin;
+				break;
+			}
+		}
+		check(OutputPin);
+
+		UEdGraphPin* InputPin = nullptr;
+		int32 Inputs = 0;
+		for (UEdGraphPin* ActivePin : OutNodeGraph->Pins)
+		{
+			if (ActivePin && ActivePin->Direction == EGPD_Input)
+			{
+				if (Inputs == InputPinIndex)
+				{
+					InputPin = ActivePin;
+					break;
+				}
+				Inputs++;
+			}
+		}
+
+		if (InputPinDisplayName)
+		{
+			InputPin->PinFriendlyName = *InputPinDisplayName;
+		}
+		InputPin->MakeLinkTo(OutputPin);
+	}
+#endif // WITH_EDITOR
+};
+
diff --git a/Engine/Source/Runtime/Engine/Private/AudioSettings.cpp b/Engine/Source/Runtime/Engine/Private/AudioSettings.cpp
--- a/Engine/Source/Runtime/Engine/Private/AudioSettings.cpp
+++ b/Engine/Source/Runtime/Engine/Private/AudioSettings.cpp
@@ -1,7 +1,7 @@
 // Copyright 1998-2019 Epic Games, Inc. All Rights Reserved.
 
 /*=============================================================================
-	PlayerInput.cpp: Unreal input system.
+	AudioSettings.cpp: Unreal audio settings
 =============================================================================*/
 
 #include "Sound/AudioSettings.h"
@@ -135,6 +135,8 @@ void UAudioSettings::PostEditChangeChainProperty(FPropertyChangedChainEvent& Pro
 				It->ReconcileNode(true);
 			}
 		}
+
+		AudioSettingsChanged.Broadcast();
 	}
 }
 #endif
@@ -159,6 +161,12 @@ const FAudioQualitySettings& UAudioSettings::GetQualityLevelSettings(int32 Quali
 	return QualityLevels[FMath::Clamp(QualityLevel, 0, QualityLevels.Num() - 1)];
 }
 
+int32 UAudioSettings::GetQualityLevelSettingsNum() const
+{
+	return QualityLevels.Num();
+}
+
+
 void UAudioSettings::SetAudioMixerEnabled(const bool bInAudioMixerEnabled)
 {
 	bIsAudioMixerEnabled = bInAudioMixerEnabled;
diff --git a/Engine/Source/Runtime/Engine/Private/SoundCue.cpp b/Engine/Source/Runtime/Engine/Private/SoundCue.cpp
--- a/Engine/Source/Runtime/Engine/Private/SoundCue.cpp
+++ b/Engine/Source/Runtime/Engine/Private/SoundCue.cpp
@@ -10,24 +10,27 @@
 #include "UObject/UObjectIterator.h"
 #include "EngineUtils.h"
 #include "Sound/SoundNode.h"
-#include "Sound/SoundNodeMixer.h"
 #include "Sound/SoundNodeAssetReferencer.h"
+#include "Sound/SoundNodeMixer.h"
 #include "Sound/SoundWave.h"
-#include "Sound/SoundNodeWavePlayer.h"
 #include "Sound/SoundNodeAttenuation.h"
+#include "Sound/SoundNodeModulator.h"
 #include "Sound/SoundNodeQualityLevel.h"
-#include "Sound/SoundNodeSoundClass.h"
 #include "Sound/SoundNodeRandom.h"
+#include "Sound/SoundNodeSoundClass.h"
+#include "Sound/SoundNodeWavePlayer.h"
 #include "GameFramework/GameUserSettings.h"
 #include "AudioCompressionSettingsUtils.h"
 #include "AudioThread.h"
+#include "DSP/Dsp.h"
 #if WITH_EDITOR
 #include "Kismet2/BlueprintEditorUtils.h"
+#include "Sound/AudioSettings.h"
 #include "SoundCueGraph/SoundCueGraphNode.h"
 #include "SoundCueGraph/SoundCueGraph.h"
 #include "SoundCueGraph/SoundCueGraphNode_Root.h"
 #include "SoundCueGraph/SoundCueGraphSchema.h"
-#endif
+#endif // WITH_EDITOR
 
 /*-----------------------------------------------------------------------------
 	USoundCue implementation.
@@ -37,7 +40,7 @@ int32 USoundCue::CachedQualityLevel = -1;
 
 #if WITH_EDITOR
 TSharedPtr<ISoundCueAudioEditor> USoundCue::SoundCueAudioEditor = nullptr;
-#endif
+#endif // WITH_EDITOR
 
 USoundCue::USoundCue(const FObjectInitializer& ObjectInitializer)
 	: Super(ObjectInitializer)
@@ -59,7 +62,6 @@ void USoundCue::PostInitProperties()
 	CacheAggregateValues();
 }
 
-
 void USoundCue::AddReferencedObjects(UObject* InThis, FReferenceCollector& Collector)
 {
 	USoundCue* This = CastChecked<USoundCue>(InThis);
@@ -106,14 +108,14 @@ void USoundCue::Serialize(FStructuredArchive::FRecord Record)
 		{
 			Record << NAMED_FIELD(SoundCueGraph);
 		}
-#endif
+#endif // WITH_EDITORONLY_DATA
 	}
 #if WITH_EDITOR
 	else
 	{
 		Record << NAMED_FIELD(SoundCueGraph);
 	}
-#endif
+#endif // WITH_EDITOR
 }
 
 void USoundCue::PostLoad()
@@ -140,7 +142,7 @@ void USoundCue::PostLoad()
 		}
 	}
 	else
-#endif
+#endif // WITH_EDITOR
 	if (GEngine && *GEngine->GameUserSettingsClass)
 	{
 		EvaluateNodes(false);
@@ -278,7 +280,7 @@ void USoundCue::PostEditChangeProperty(struct FPropertyChangedEvent& PropertyCha
 
 	CacheAggregateValues();
 }
-#endif
+#endif // WITH_EDITOR
 
 void USoundCue::RecursiveFindAttenuation( USoundNode* Node, TArray<class USoundNodeAttenuation*> &OutNodes )
 {
@@ -630,11 +632,20 @@ void USoundCue::SetSoundCueAudioEditor(TSharedPtr<ISoundCueAudioEditor> InSoundC
 	SoundCueAudioEditor = InSoundCueAudioEditor;
 }
 
+void USoundCue::ResetGraph()
+{
+	for (const USoundNode* Node : AllNodes)
+	{
+		SoundCueGraph->RemoveNode(Node->GraphNode);
+	}
+
+	AllNodes.Reset();
+	FirstNode = nullptr;
+}
+
 /** Gets the sound cue graph editor implementation. */
 TSharedPtr<ISoundCueAudioEditor> USoundCue::GetSoundCueAudioEditor()
-	{
+{
 	return SoundCueAudioEditor;
 }
-
-
-#endif
+#endif // WITH_EDITOR
diff --git a/Engine/Source/Runtime/Engine/Private/SoundCueTemplate.cpp b/Engine/Source/Runtime/Engine/Private/SoundCueTemplate.cpp
new file mode 100644
--- /dev/null
+++ b/Engine/Source/Runtime/Engine/Private/SoundCueTemplate.cpp
@@ -0,0 +1,32 @@
+// Copyright 1998-2019 Epic Games, Inc. All Rights Reserved.
+
+#include "Sound/SoundCueTemplate.h"
+
+
+USoundCueTemplate::USoundCueTemplate(const FObjectInitializer& ObjectInitializer)
+	: Super(ObjectInitializer)
+{
+}
+
+#if WITH_EDITOR
+void USoundCueTemplate::PostEditChangeProperty(struct FPropertyChangedEvent& PropertyChangedEvent)
+{
+	Super::PostEditChangeProperty(PropertyChangedEvent);
+
+	if (PropertyChangedEvent.Property)
+	{
+		RebuildGraph(*this);
+	}
+}
+
+void USoundCueTemplate::RebuildGraph(USoundCue& SoundCue) const
+{
+	SoundCue.ResetGraph();
+
+	OnRebuildGraph(SoundCue);
+
+	SoundCue.LinkGraphNodesFromSoundNodes();
+	SoundCue.PostEditChange();
+	SoundCue.MarkPackageDirty();
+}
+#endif // WITH_EDITOR
