Mark AudioModulation & SoundCueTemplates as Beta

diff --git a/Engine/Plugins/Runtime/AudioModulation/AudioModulation.uplugin b/Engine/Plugins/Runtime/AudioModulation/AudioModulation.uplugin
--- a/Engine/Plugins/Runtime/AudioModulation/AudioModulation.uplugin
+++ b/Engine/Plugins/Runtime/AudioModulation/AudioModulation.uplugin
@@ -12,7 +12,7 @@
 	"SupportURL": "",
 	"EnabledByDefault": false,
 	"CanContainContent": false,
-	"IsBetaVersion": false,
+	"IsBetaVersion": true,
 	"Installed": false,
 	"Modules": [
 		{
diff --git a/Engine/Plugins/Runtime/SoundCueTemplates/SoundCueTemplates.uplugin b/Engine/Plugins/Runtime/SoundCueTemplates/SoundCueTemplates.uplugin
--- a/Engine/Plugins/Runtime/SoundCueTemplates/SoundCueTemplates.uplugin
+++ b/Engine/Plugins/Runtime/SoundCueTemplates/SoundCueTemplates.uplugin
@@ -11,7 +11,7 @@
 	"MarketplaceURL": "",
 	"SupportURL": "",
 	"CanContainContent": false,
-	"IsBetaVersion": false,
+	"IsBetaVersion": true,
 	"Installed": false,
 	"Modules": [
 		{
