Moving SoundCueTemplate code out of Editor/AudioEditor and Engine
+ Addition of a new Editor module inside existing SoundCueTemplate plugin

diff --git a/Engine/Plugins/Runtime/SoundCueTemplates/SoundCueTemplates.uplugin b/Engine/Plugins/Runtime/SoundCueTemplates/SoundCueTemplates.uplugin
--- a/Engine/Plugins/Runtime/SoundCueTemplates/SoundCueTemplates.uplugin
+++ b/Engine/Plugins/Runtime/SoundCueTemplates/SoundCueTemplates.uplugin
@@ -6,7 +6,7 @@
 	"Description": "Collection of SoundCue Templates, which provide rapid design of common audio design workflows.",
 	"Category": "Audio",
 	"CreatedBy": "Epic Games, Inc.",
-	"CreatedByURL": "",
+	"CreatedByURL": "http://epicgames.com",
 	"DocsURL": "",
 	"MarketplaceURL": "",
 	"SupportURL": "",
@@ -18,6 +18,12 @@
 			"Name": "SoundCueTemplates",
 			"Type": "Runtime",
 			"LoadingPhase": "Default"
+		},
+		{
+			"Name" : "SoundCueTemplatesEditor",
+			"Type" : "Editor",
+			"LoadingPhase" : "PostDefault"
 		}
+		
 	]
 }
\ No newline at end of file
diff --git a/Engine/Plugins/Runtime/SoundCueTemplates/Source/SoundCueTemplates/Private/SoundCueContainer.cpp b/Engine/Plugins/Runtime/SoundCueTemplates/Source/SoundCueTemplates/Private/SoundCueContainer.cpp
--- a/Engine/Plugins/Runtime/SoundCueTemplates/Source/SoundCueTemplates/Private/SoundCueContainer.cpp
+++ b/Engine/Plugins/Runtime/SoundCueTemplates/Source/SoundCueTemplates/Private/SoundCueContainer.cpp
@@ -2,7 +2,6 @@
 #include "SoundCueContainer.h"
 
 #include "Sound/SoundCue.h"
-#include "Sound/SoundCueTemplate.h"
 #include "Sound/SoundWave.h"
 #include "Sound/SoundNode.h"
 #include "Sound/SoundNodeConcatenator.h"
@@ -11,7 +10,6 @@
 #include "Sound/SoundNodeQualityLevel.h"
 #include "Sound/SoundNodeRandom.h"
 #include "Sound/SoundNodeWavePlayer.h"
-#include "SoundCueTemplateSettings.h"
 #include "SoundCueTemplatesModule.h"
 
 #if WITH_EDITORONLY_DATA
diff --git a/Engine/Plugins/Runtime/SoundCueTemplates/Source/SoundCueTemplates/Private/SoundCueTemplate.cpp b/Engine/Plugins/Runtime/SoundCueTemplates/Source/SoundCueTemplates/Private/SoundCueTemplate.cpp
new file mode 100644
--- /dev/null
+++ b/Engine/Plugins/Runtime/SoundCueTemplates/Source/SoundCueTemplates/Private/SoundCueTemplate.cpp
@@ -0,0 +1,31 @@
+// Copyright 1998-2019 Epic Games, Inc. All Rights Reserved.
+
+#include "SoundCueTemplate.h"
+
+USoundCueTemplate::USoundCueTemplate(const FObjectInitializer& ObjectInitializer)
+	: Super(ObjectInitializer)
+{
+}
+
+#if WITH_EDITOR
+void USoundCueTemplate::PostEditChangeProperty(struct FPropertyChangedEvent& PropertyChangedEvent)
+{
+	Super::PostEditChangeProperty(PropertyChangedEvent);
+
+	if (PropertyChangedEvent.Property)
+	{
+		RebuildGraph(*this);
+	}
+}
+
+void USoundCueTemplate::RebuildGraph(USoundCue& SoundCue) const
+{
+	SoundCue.ResetGraph();
+
+	OnRebuildGraph(SoundCue);
+
+	SoundCue.LinkGraphNodesFromSoundNodes();
+	SoundCue.PostEditChange();
+	SoundCue.MarkPackageDirty();
+}
+#endif // WITH_EDITOR
diff --git a/Engine/Plugins/Runtime/SoundCueTemplates/Source/SoundCueTemplates/Public/SoundCueContainer.h b/Engine/Plugins/Runtime/SoundCueTemplates/Source/SoundCueTemplates/Public/SoundCueContainer.h
--- a/Engine/Plugins/Runtime/SoundCueTemplates/Source/SoundCueTemplates/Public/SoundCueContainer.h
+++ b/Engine/Plugins/Runtime/SoundCueTemplates/Source/SoundCueTemplates/Public/SoundCueContainer.h
@@ -3,10 +3,8 @@
 #pragma once
 
 #include "CoreMinimal.h"
-
-
 #include "Misc/Attribute.h"
-#include "Sound/SoundCueTemplate.h"
+#include "SoundCueTemplate.h"
 #include "SoundCueTemplateSettings.h"
 #include "UObject/ObjectMacros.h"
 
diff --git a/Engine/Plugins/Runtime/SoundCueTemplates/Source/SoundCueTemplates/Public/SoundCueDistanceCrossfade.h b/Engine/Plugins/Runtime/SoundCueTemplates/Source/SoundCueTemplates/Public/SoundCueDistanceCrossfade.h
--- a/Engine/Plugins/Runtime/SoundCueTemplates/Source/SoundCueTemplates/Public/SoundCueDistanceCrossfade.h
+++ b/Engine/Plugins/Runtime/SoundCueTemplates/Source/SoundCueTemplates/Public/SoundCueDistanceCrossfade.h
@@ -3,10 +3,8 @@
 #pragma once
 
 #include "CoreMinimal.h"
-
-#include "Sound/SoundCueTemplate.h"
+#include "SoundCueTemplate.h"
 #include "Sound/SoundNodeDistanceCrossFade.h"
-
 #include "Misc/Attribute.h"
 #include "UObject/ObjectMacros.h"
 
diff --git a/Engine/Plugins/Runtime/SoundCueTemplates/Source/SoundCueTemplates/Public/SoundCueTemplate.h b/Engine/Plugins/Runtime/SoundCueTemplates/Source/SoundCueTemplates/Public/SoundCueTemplate.h
new file mode 100644
--- /dev/null
+++ b/Engine/Plugins/Runtime/SoundCueTemplates/Source/SoundCueTemplates/Public/SoundCueTemplate.h
@@ -0,0 +1,180 @@
+// Copyright 1998-2019 Epic Games, Inc. All Rights Reserved.
+#pragma once
+
+#include "CoreMinimal.h"
+
+#include "UObject/ObjectMacros.h"
+#include "Sound/SoundCue.h"
+#include "Sound/SoundNode.h"
+#include "Templates/SubclassOf.h"
+
+#if WITH_EDITOR
+#include "EdGraph/EdGraphPin.h"
+#endif // WITH_EDITOR
+
+#include "SoundCueTemplate.generated.h"
+
+
+// Base Sound Cue Template class, which builds the sound node graph procedurally and hides more complex Sound Cue functionality
+// to streamline implementation defined in child classes.
+UCLASS(abstract)
+class SOUNDCUETEMPLATES_API USoundCueTemplate : public USoundCue
+{
+	GENERATED_UCLASS_BODY()
+
+#if WITH_EDITOR
+	virtual void PostEditChangeProperty(FPropertyChangedEvent& PropertyChangedEvent) override;
+
+public:
+	/** Rebuilds the graph when user-facing properties are changed */
+	void RebuildGraph(USoundCue& SoundCue) const;
+
+protected:
+	/**
+	 * Function to override that uses internal editor-data only properties to re-build node
+	 * graph of provided SoundCue.  Provided SoundCue may be a reference to this SoundCue
+	 * or one to be copied from this template's graph-building schema.
+	 */
+	virtual void OnRebuildGraph(USoundCue& SoundCue) const PURE_VIRTUAL(, );
+
+	/**
+	 * Node width offset between nodes in the SoundCueEditor.
+	 */
+	static float GetNodeWidthOffset()
+	{
+		return 500.0f;
+	}
+
+	/**
+	 * Initial width offset of first node in the SoundCueEditor.
+	 */
+	static float GetInitialWidthOffset()
+	{
+		return 300.0f;
+	}
+
+	/**
+	 * Node height offset between nodes in the SoundCueEditor.
+	 */
+	static float GetNodeHeightOffset()
+	{
+		return 200.0f;
+	}
+
+	/**
+	 * Initial height offset of first node in the SoundCueEditor.
+	 */
+	static float GetInitialHeightOffset()
+	{
+		return 0.0f;
+	}
+
+	/**
+	 * Sets the provided node's position to the location on the template's grid
+	 */
+	static void SetSoundNodePosition(USoundNode& InNode, int32 Column, int32 Row)
+	{
+		InNode.GraphNode->NodePosX -= Column * GetNodeWidthOffset() + GetInitialWidthOffset();
+		InNode.GraphNode->NodePosY += Row * GetNodeHeightOffset() + GetInitialHeightOffset();
+	}
+
+	/**
+	 * Utility function that returns newly constructed sound node that has been added to provided SoundCue's Node Graph.
+	 * SoundCue				- SoundCue to insert the constructed node into.
+	 * InParentNode			- If not null, Node to set as parent of constructed child.
+	 * Column				- (Optional) If positive, sets visual column of node in built graph as displayed in the SoundCueEditor.  Columns are indexed as positive
+	 *						  to the left from the root (index 0).  Can be viewed when using the 'Copy To SoundCue' action of the given template.
+	 * Row					- (Optional) If positive, sets visual row of node in built graph as displayed in the SoundCueEditor.  Columns are indexed as positive
+	 *						  descending from the root (index 0). Can be viewed when using the 'Copy To SoundCue' action of the given template.
+	 * InputPinIndex		- (Optional) If InParentNode set, supplies index of input pin to connect newly constructed ChildNode to.
+	 * InputPinDisplayName	- (Optional) Name to set the targeted input pin to.
+	 */
+	template <class T>
+	static T& ConstructSoundNodeChild(USoundCue& SoundCue, USoundNode* InParentNode, int32 Column, int32 Row, int32 InputPinIndex = -1, const FText* InputPinDisplayName = nullptr)
+	{
+		T* ChildNode = SoundCue.ConstructSoundNode<T>();
+		check(ChildNode);
+
+		SetSoundNodePosition(*ChildNode, Column, Row);
+
+		if (InParentNode)
+		{
+			AddSoundNodeChild(*InParentNode, *ChildNode, InputPinIndex, InputPinDisplayName);
+		}
+
+		return *ChildNode;
+	}
+
+	/**
+	 * Constructs initial root sound node of template. Function will fail if SoundCue already contains a root/first node.
+	 * SoundCue	- Parent of initial node to construct.
+	 */
+	template <class T>
+	static T& ConstructSoundNodeRoot(USoundCue& SoundCue)
+	{
+		T* RootNode = SoundCue.ConstructSoundNode<T>();
+		check(RootNode);
+		check(!SoundCue.FirstNode);
+
+		SoundCue.FirstNode = RootNode;
+		RootNode->GraphNode->NodePosX = -1 * GetInitialWidthOffset();
+		RootNode->GraphNode->NodePosY = GetInitialHeightOffset();
+
+		return *RootNode;
+	}
+
+	/**
+	 * Utility function that adds a provided child node to a parent node's list of children.
+	 * InParentNode			- Node to set as parent of provided child.
+	 * InChildNode			- Node to set as child of provided parent.
+	 * InputPinIndex		- Index of input pin to connect newly constructed ChildNode to.
+	 * InputPinDisplayName	- (Optional) Name to set the targeted input pin to.
+	 */
+	static void AddSoundNodeChild(USoundNode& InParentNode, USoundNode& InChildNode, int32 InputPinIndex, const FText* InputPinDisplayName = nullptr)
+	{
+		check(InputPinIndex >= 0 && InputPinIndex <= InParentNode.GetMaxChildNodes());
+
+		const int32 NewChildNodeIndex = InParentNode.ChildNodes.Num();
+		InParentNode.InsertChildNode(NewChildNodeIndex);
+		InParentNode.ChildNodes[NewChildNodeIndex] = &InChildNode;
+
+		UEdGraphNode* OutNodeGraph = InParentNode.GetGraphNode();
+		check(OutNodeGraph);
+		UEdGraphPin* OutputPin = nullptr;
+		for (UEdGraphPin* ActivePin : InChildNode.GetGraphNode()->Pins)
+		{
+			if (ActivePin->Direction == EGPD_Output)
+			{
+				OutputPin = ActivePin;
+				break;
+			}
+		}
+		check(OutputPin);
+
+		UEdGraphPin* InputPin = nullptr;
+		int32 Inputs = 0;
+		for (UEdGraphPin* ActivePin : OutNodeGraph->Pins)
+		{
+			if (ActivePin && ActivePin->Direction == EGPD_Input)
+			{
+				if (Inputs == InputPinIndex)
+				{
+					InputPin = ActivePin;
+					break;
+				}
+				Inputs++;
+			}
+		}
+
+		if (InputPin)
+		{
+			if (InputPinDisplayName)
+			{
+				InputPin->PinFriendlyName = *InputPinDisplayName;
+			}
+			InputPin->MakeLinkTo(OutputPin);
+		}
+	}
+#endif // WITH_EDITOR
+};
+
diff --git a/Engine/Plugins/Runtime/SoundCueTemplates/Source/SoundCueTemplates/Public/SoundCueTemplatesModule.h b/Engine/Plugins/Runtime/SoundCueTemplates/Source/SoundCueTemplates/Public/SoundCueTemplatesModule.h
--- a/Engine/Plugins/Runtime/SoundCueTemplates/Source/SoundCueTemplates/Public/SoundCueTemplatesModule.h
+++ b/Engine/Plugins/Runtime/SoundCueTemplates/Source/SoundCueTemplates/Public/SoundCueTemplatesModule.h
@@ -5,7 +5,7 @@
 #include "CoreMinimal.h"
 #include "Modules/ModuleManager.h"
 
-DECLARE_LOG_CATEGORY_EXTERN(SoundCueTemplates, Log, All);
+SOUNDCUETEMPLATES_API DECLARE_LOG_CATEGORY_EXTERN(SoundCueTemplates, Log, All);
 
 class FSoundCueTemplatesModule : public IModuleInterface
 {
diff --git a/Engine/Plugins/Runtime/SoundCueTemplates/Source/SoundCueTemplatesEditor/Private/AssetTypeActions_SoundCueTemplate.cpp b/Engine/Plugins/Runtime/SoundCueTemplates/Source/SoundCueTemplatesEditor/Private/AssetTypeActions_SoundCueTemplate.cpp
new file mode 100644
--- /dev/null
+++ b/Engine/Plugins/Runtime/SoundCueTemplates/Source/SoundCueTemplatesEditor/Private/AssetTypeActions_SoundCueTemplate.cpp
@@ -0,0 +1,55 @@
+// Copyright 1998-2019 Epic Games, Inc. All Rights Reserved.
+#include "AssetTypeActions_SoundCueTemplate.h"
+
+#include "SoundCueTemplate.h"
+#include "SoundCueTemplateFactory.h"
+
+#include "ToolMenuSection.h"
+#include "ContentBrowserModule.h"
+#include "IContentBrowserSingleton.h"
+
+#define LOCTEXT_NAMESPACE "AssetTypeActions"
+
+UClass* FAssetTypeActions_SoundCueTemplate::GetSupportedClass() const
+{
+	return USoundCueTemplate::StaticClass();
+}
+
+void FAssetTypeActions_SoundCueTemplate::GetActions(const TArray<UObject*>& InObjects, FToolMenuSection& Section)
+{
+	TArray<TWeakObjectPtr<USoundCueTemplate>> Cues = GetTypedWeakObjectPtrs<USoundCueTemplate>(InObjects);
+
+	Section.AddMenuEntry(
+		"SoundCueTemplate_CopyToSoundCue",
+		LOCTEXT("SoundCueTemplate_CopyToSoundCue", "Copy To Sound Cue"),
+		LOCTEXT("SoundCueTemplate_CopyToSoundCueTooltip", "Exports a Sound Cue Template to a Sound Cue."),
+		FSlateIcon(FEditorStyle::GetStyleSetName(), "ClassIcon.SoundCue"),
+		FUIAction(
+			FExecuteAction::CreateSP(this, &FAssetTypeActions_SoundCueTemplate::ExecuteCopyToSoundCue, Cues),
+			FCanExecuteAction()
+		)
+	);
+}
+
+void FAssetTypeActions_SoundCueTemplate::ExecuteCopyToSoundCue(TArray<TWeakObjectPtr<USoundCueTemplate>> Objects)
+{
+	for (const TWeakObjectPtr<USoundCueTemplate>& Object : Objects)
+	{
+		if (!Object.IsValid())
+		{
+			continue;
+		}
+
+		FString Name;
+		FString PackagePath;
+		CreateUniqueAssetName(Object->GetOutermost()->GetName(), TEXT(""), PackagePath, Name);
+
+		if (USoundCueTemplateCopyFactory* Factory = NewObject<USoundCueTemplateCopyFactory>())
+		{
+			Factory->SoundCueTemplate = Object;
+			FContentBrowserModule& ContentBrowserModule = FModuleManager::LoadModuleChecked<FContentBrowserModule>("ContentBrowser");
+			ContentBrowserModule.Get().CreateNewAsset(Name, FPackageName::GetLongPackagePath(PackagePath), USoundCue::StaticClass(), Factory);
+		}
+	}
+}
+#undef LOCTEXT_NAMESPACE
diff --git a/Engine/Plugins/Runtime/SoundCueTemplates/Source/SoundCueTemplatesEditor/Private/AssetTypeActions_SoundCueTemplate.h b/Engine/Plugins/Runtime/SoundCueTemplates/Source/SoundCueTemplatesEditor/Private/AssetTypeActions_SoundCueTemplate.h
new file mode 100644
--- /dev/null
+++ b/Engine/Plugins/Runtime/SoundCueTemplates/Source/SoundCueTemplatesEditor/Private/AssetTypeActions_SoundCueTemplate.h
@@ -0,0 +1,25 @@
+// Copyright 1998-2019 Epic Games, Inc. All Rights Reserved.
+#pragma once
+
+#include "CoreMinimal.h"
+#include "Toolkits/IToolkitHost.h"
+#include "AssetTypeActions_Base.h"
+
+class USoundCueTemplate;
+
+class FAssetTypeActions_SoundCueTemplate : public FAssetTypeActions_Base
+{
+public:
+	// IAssetTypeActions Implementation
+	virtual FText GetName() const override { return NSLOCTEXT("AssetTypeActions", "AssetTypeActions_SoundCueTemplate", "Sound Cue Template"); }
+	virtual FColor GetTypeColor() const override { return FColor(255, 255, 255); }
+	virtual UClass* GetSupportedClass() const override;
+	virtual bool HasActions(const TArray<UObject*>& InObjects) const override { return true; }
+	virtual void GetActions(const TArray<UObject*>& InObjects, struct FToolMenuSection& Section) override;
+	virtual uint32 GetCategories() override { return EAssetTypeCategories::Sounds; }
+	virtual bool CanFilter() override { return true; }
+
+protected:
+	/** Converts the provided SoundCue Template to a fully-modifiable SoundCue */
+	void ExecuteCopyToSoundCue(TArray<TWeakObjectPtr<USoundCueTemplate>> Objects);
+};
diff --git a/Engine/Plugins/Runtime/SoundCueTemplates/Source/SoundCueTemplatesEditor/Private/SoundCueTemplateClassFilter.h b/Engine/Plugins/Runtime/SoundCueTemplates/Source/SoundCueTemplatesEditor/Private/SoundCueTemplateClassFilter.h
new file mode 100644
--- /dev/null
+++ b/Engine/Plugins/Runtime/SoundCueTemplates/Source/SoundCueTemplatesEditor/Private/SoundCueTemplateClassFilter.h
@@ -0,0 +1,28 @@
+// Copyright 1998-2019 Epic Games, Inc. All Rights Reserved.
+
+#pragma once
+
+#include "CoreMinimal.h"
+#include "UObject/Class.h"
+#include "UObject/ObjectMacros.h"
+#include "ClassViewerFilter.h"
+
+class FSoundCueTemplateAssetParentFilter : public IClassViewerFilter
+{
+public:
+	/** All children of these classes will be included unless filtered out by another setting. */
+	TSet< const UClass* > AllowedChildrenOfClasses;
+
+	/** Disallowed class flags. */
+	EClassFlags DisallowedClassFlags;
+
+	virtual bool IsClassAllowed(const FClassViewerInitializationOptions& InInitOptions, const UClass* InClass, TSharedRef< FClassViewerFilterFuncs > InFilterFuncs) override
+	{
+		return !InClass->HasAnyClassFlags(DisallowedClassFlags) && InFilterFuncs->IfInChildOfClassesSet(AllowedChildrenOfClasses, InClass) != EFilterReturn::Failed;
+	}
+
+	virtual bool IsUnloadedClassAllowed(const FClassViewerInitializationOptions& InInitOptions, const TSharedRef< const IUnloadedBlueprintData > InUnloadedClassData, TSharedRef< FClassViewerFilterFuncs > InFilterFuncs) override
+	{
+		return !InUnloadedClassData->HasAnyClassFlags(DisallowedClassFlags) && InFilterFuncs->IfInChildOfClassesSet(AllowedChildrenOfClasses, InUnloadedClassData) != EFilterReturn::Failed;
+	}
+};
\ No newline at end of file
diff --git a/Engine/Plugins/Runtime/SoundCueTemplates/Source/SoundCueTemplatesEditor/Private/SoundCueTemplateFactory.cpp b/Engine/Plugins/Runtime/SoundCueTemplates/Source/SoundCueTemplatesEditor/Private/SoundCueTemplateFactory.cpp
new file mode 100644
--- /dev/null
+++ b/Engine/Plugins/Runtime/SoundCueTemplates/Source/SoundCueTemplatesEditor/Private/SoundCueTemplateFactory.cpp
@@ -0,0 +1,89 @@
+// Copyright 1998-2019 Epic Games, Inc. All Rights Reserved.
+#include "SoundCueTemplateFactory.h"
+#include "SoundCueTemplateClassFilter.h"
+
+#include "ClassViewerModule.h"
+#include "EdGraph/EdGraphPin.h"
+#include "Engine/Engine.h"
+#include "Kismet2/SClassPickerDialog.h"
+#include "SoundCueGraph/SoundCueGraphNode.h"
+#include "Sound/SoundCue.h"
+#include "UObject/Class.h"
+
+#define LOCTEXT_NAMESPACE "SoundCueTemplatesEditor"
+
+USoundCueTemplateCopyFactory::USoundCueTemplateCopyFactory(const FObjectInitializer& ObjectInitializer)
+	: Super(ObjectInitializer)
+{
+	SupportedClass = USoundCue::StaticClass();
+
+	bCreateNew = false;
+	bEditorImport = false;
+	bEditAfterNew = true;
+}
+
+UObject* USoundCueTemplateCopyFactory::FactoryCreateNew(UClass* Class, UObject* InParent, FName Name, EObjectFlags Flags, UObject* Context, FFeedbackContext* Warn)
+{
+	if (SoundCueTemplate.IsValid())
+	{
+		if (USoundCue* NewCue = NewObject<USoundCue>(InParent, Name, Flags))
+		{
+			SoundCueTemplate.Get()->RebuildGraph(*NewCue);
+			return NewCue;
+		}
+	}
+
+	return nullptr;
+}
+
+USoundCueTemplateFactory::USoundCueTemplateFactory(const FObjectInitializer& ObjectInitializer)
+	: Super(ObjectInitializer)
+{
+	SupportedClass = USoundCueTemplate::StaticClass();
+
+	bCreateNew = true;
+	bEditorImport = false;
+	bEditAfterNew = true;
+	SoundCueTemplateClass = nullptr;
+}
+
+bool USoundCueTemplateFactory::ConfigureProperties()
+{
+	SoundCueTemplateClass = nullptr;
+
+	// Load the classviewer module to display a class picker
+	FClassViewerModule& ClassViewerModule = FModuleManager::LoadModuleChecked<FClassViewerModule>("ClassViewer");
+
+	TSharedPtr<FSoundCueTemplateAssetParentFilter> Filter = MakeShareable(new FSoundCueTemplateAssetParentFilter);
+	Filter->DisallowedClassFlags = CLASS_Abstract | CLASS_Deprecated | CLASS_NewerVersionExists | CLASS_HideDropDown;
+	Filter->AllowedChildrenOfClasses.Add(USoundCueTemplate::StaticClass());
+
+	FClassViewerInitializationOptions Options;
+	Options.Mode = EClassViewerMode::ClassPicker;
+	Options.ClassFilter = Filter;
+
+	const FText TitleText = LOCTEXT("CreateSoundCueTemplateOfType", "Pick Type of SoundCueTemplate");
+	UClass* ChosenClass = nullptr;
+
+	const bool bPressedOk = SClassPickerDialog::PickClass(TitleText, Options, ChosenClass, USoundCueTemplate::StaticClass());
+
+	if (bPressedOk)
+	{
+		SoundCueTemplateClass = ChosenClass;
+	}
+	return bPressedOk;
+}
+
+UObject* USoundCueTemplateFactory::FactoryCreateNew(UClass* Class, UObject* InParent, FName Name, EObjectFlags Flags, UObject* Context, FFeedbackContext* Warn)
+{
+	if (SoundCueTemplateClass != nullptr)
+	{
+		if (USoundCueTemplate* NewSoundCueTemplate = NewObject<USoundCueTemplate>(InParent, SoundCueTemplateClass, Name, Flags))
+		{
+			NewSoundCueTemplate->RebuildGraph(*NewSoundCueTemplate);
+			return NewSoundCueTemplate;
+		}
+	}
+
+	return nullptr;
+}
diff --git a/Engine/Plugins/Runtime/SoundCueTemplates/Source/SoundCueTemplatesEditor/Private/SoundCueTemplatesEditorModule.cpp b/Engine/Plugins/Runtime/SoundCueTemplates/Source/SoundCueTemplatesEditor/Private/SoundCueTemplatesEditorModule.cpp
new file mode 100644
--- /dev/null
+++ b/Engine/Plugins/Runtime/SoundCueTemplates/Source/SoundCueTemplatesEditor/Private/SoundCueTemplatesEditorModule.cpp
@@ -0,0 +1,27 @@
+// Copyright 1998-2019 Epic Games, Inc. All Rights Reserved.
+
+#include "SoundCueTemplatesEditorModule.h"
+
+#include "Modules/ModuleManager.h"
+#include "AssetToolsModule.h"
+
+#include "AssetTypeActions_SoundCueTemplate.h"
+
+void FSoundCueTemplatesEditorModule::StartupModule()
+{
+	RegisterAssetActions();
+}
+
+void FSoundCueTemplatesEditorModule::ShutdownModule()
+{
+}
+
+void FSoundCueTemplatesEditorModule::RegisterAssetActions()
+{
+	// Register the audio editor asset type actions
+	IAssetTools& AssetTools = FModuleManager::LoadModuleChecked<FAssetToolsModule>("AssetTools").Get();
+
+	AssetTools.RegisterAssetTypeActions(MakeShared<FAssetTypeActions_SoundCueTemplate>());
+}
+
+IMPLEMENT_MODULE(FSoundCueTemplatesEditorModule, SoundCueTemplatesEditor);
diff --git a/Engine/Plugins/Runtime/SoundCueTemplates/Source/SoundCueTemplatesEditor/Public/SoundCueTemplateFactory.h b/Engine/Plugins/Runtime/SoundCueTemplates/Source/SoundCueTemplatesEditor/Public/SoundCueTemplateFactory.h
new file mode 100644
--- /dev/null
+++ b/Engine/Plugins/Runtime/SoundCueTemplates/Source/SoundCueTemplatesEditor/Public/SoundCueTemplateFactory.h
@@ -0,0 +1,37 @@
+// Copyright 1998-2019 Epic Games, Inc. All Rights Reserved.
+#pragma once
+
+#include "CoreMinimal.h"
+#include "UObject/ObjectMacros.h"
+#include "SoundCueTemplate.h"
+#include "Factories/Factory.h"
+
+#include "SoundCueTemplateFactory.generated.h"
+
+UCLASS(hidecategories = Object, MinimalAPI)
+class USoundCueTemplateCopyFactory : public UFactory
+{
+	GENERATED_UCLASS_BODY()
+
+	//~ Begin UFactory Interface
+	virtual UObject* FactoryCreateNew(UClass* Class, UObject* InParent, FName Name, EObjectFlags Flags, UObject* Context, FFeedbackContext* Warn) override;
+	//~ End UFactory Interface
+
+public:
+	UPROPERTY()
+	TWeakObjectPtr<USoundCueTemplate> SoundCueTemplate;
+};
+
+UCLASS(hidecategories = Object, MinimalAPI)
+class USoundCueTemplateFactory : public UFactory
+{
+	GENERATED_UCLASS_BODY()
+
+	UPROPERTY(EditAnywhere, Category = SoundCueTemplateFacotry)
+	TSubclassOf<USoundCueTemplate> SoundCueTemplateClass;
+
+	//~ Begin UFactory Interface
+	virtual bool ConfigureProperties() override;
+	virtual UObject* FactoryCreateNew(UClass* Class, UObject* InParent, FName Name, EObjectFlags Flags, UObject* Context, FFeedbackContext* Warn) override;
+	//~ End UFactory Interface
+};
\ No newline at end of file
diff --git a/Engine/Plugins/Runtime/SoundCueTemplates/Source/SoundCueTemplatesEditor/Public/SoundCueTemplatesEditorModule.h b/Engine/Plugins/Runtime/SoundCueTemplates/Source/SoundCueTemplatesEditor/Public/SoundCueTemplatesEditorModule.h
new file mode 100644
--- /dev/null
+++ b/Engine/Plugins/Runtime/SoundCueTemplates/Source/SoundCueTemplatesEditor/Public/SoundCueTemplatesEditorModule.h
@@ -0,0 +1,20 @@
+// Copyright 1998-2019 Epic Games, Inc. All Rights Reserved.
+
+#pragma once
+
+#include "CoreMinimal.h"
+#include "Containers/Set.h"
+#include "Modules/ModuleInterface.h"
+
+SOUNDCUETEMPLATESEDITOR_API DECLARE_LOG_CATEGORY_EXTERN(SoundCueTemplatesEditor, Log, All);
+
+class FSoundCueTemplatesEditorModule :	public IModuleInterface
+{
+public:
+	/** IModuleInterface implementation */
+	virtual void StartupModule() override;
+	virtual void ShutdownModule() override;
+
+private:
+	void RegisterAssetActions();
+};
diff --git a/Engine/Plugins/Runtime/SoundCueTemplates/Source/SoundCueTemplatesEditor/SoundCueTemplatesEditor.Build.cs b/Engine/Plugins/Runtime/SoundCueTemplates/Source/SoundCueTemplatesEditor/SoundCueTemplatesEditor.Build.cs
new file mode 100644
--- /dev/null
+++ b/Engine/Plugins/Runtime/SoundCueTemplates/Source/SoundCueTemplatesEditor/SoundCueTemplatesEditor.Build.cs
@@ -0,0 +1,38 @@
+// Copyright 1998-2019 Epic Games, Inc. All Rights Reserved.
+
+using UnrealBuildTool;
+
+public class SoundCueTemplatesEditor : ModuleRules
+{
+	public SoundCueTemplatesEditor(ReadOnlyTargetRules Target) : base(Target)
+	{
+        PCHUsage = ModuleRules.PCHUsageMode.UseExplicitOrSharedPCHs;
+
+        PublicDependencyModuleNames.AddRange(
+            new string[]
+            {
+                "Core",
+                "SoundCueTemplates",
+                "ToolMenus"
+            }
+        );
+
+        PrivateDependencyModuleNames.AddRange(
+            new string[]
+            {
+                "CoreUObject",
+                "EditorStyle",
+                "Engine",
+                "UnrealEd",
+                "Slate",
+                "SlateCore"
+            }
+        );
+
+        PrivateIncludePathModuleNames.AddRange(
+            new string[] {
+                "AssetTools"
+            }
+        );
+    }
+}
diff --git a/Engine/Source/Editor/AudioEditor/Classes/Factories/SoundCueTemplateFactory.h b/Engine/Source/Editor/AudioEditor/Classes/Factories/SoundCueTemplateFactory.h
deleted file mode 100644
diff --git a/Engine/Source/Editor/AudioEditor/Private/AssetTypeActions/AssetTypeActions_SoundCueTemplate.cpp b/Engine/Source/Editor/AudioEditor/Private/AssetTypeActions/AssetTypeActions_SoundCueTemplate.cpp
deleted file mode 100644
diff --git a/Engine/Source/Editor/AudioEditor/Private/AssetTypeActions/AssetTypeActions_SoundCueTemplate.h b/Engine/Source/Editor/AudioEditor/Private/AssetTypeActions/AssetTypeActions_SoundCueTemplate.h
deleted file mode 100644
diff --git a/Engine/Source/Editor/AudioEditor/Private/AudioEditorModule.cpp b/Engine/Source/Editor/AudioEditor/Private/AudioEditorModule.cpp
--- a/Engine/Source/Editor/AudioEditor/Private/AudioEditorModule.cpp
+++ b/Engine/Source/Editor/AudioEditor/Private/AudioEditorModule.cpp
@@ -27,7 +27,6 @@
 #include "AssetTypeActions/AssetTypeActions_SoundBase.h"
 #include "AssetTypeActions/AssetTypeActions_SoundClass.h"
 #include "AssetTypeActions/AssetTypeActions_SoundCue.h"
-#include "AssetTypeActions/AssetTypeActions_SoundCueTemplate.h"
 #include "AssetTypeActions/AssetTypeActions_SoundMix.h"
 #include "AssetTypeActions/AssetTypeActions_SoundWave.h"
 #include "AssetTypeActions/AssetTypeActions_ReverbEffect.h"
@@ -138,7 +137,6 @@ public:
 		AssetTools.RegisterAssetTypeActions(MakeShareable(new FAssetTypeActions_SoundBase));
 		AssetTools.RegisterAssetTypeActions(MakeShareable(new FAssetTypeActions_SoundClass));
 		AssetTools.RegisterAssetTypeActions(MakeShareable(new FAssetTypeActions_SoundCue));
-		AssetTools.RegisterAssetTypeActions(MakeShareable(new FAssetTypeActions_SoundCueTemplate));
 		AssetTools.RegisterAssetTypeActions(MakeShareable(new FAssetTypeActions_SoundMix));
 		AssetTools.RegisterAssetTypeActions(MakeShareable(new FAssetTypeActions_SoundWave));
 		AssetTools.RegisterAssetTypeActions(MakeShareable(new FAssetTypeActions_ReverbEffect));
diff --git a/Engine/Source/Editor/AudioEditor/Private/Factories/SoundCueTemplateFactory.cpp b/Engine/Source/Editor/AudioEditor/Private/Factories/SoundCueTemplateFactory.cpp
deleted file mode 100644
diff --git a/Engine/Source/Editor/GameProjectGeneration/Private/GameProjectUtils.cpp b/Engine/Source/Editor/GameProjectGeneration/Private/GameProjectUtils.cpp
--- a/Engine/Source/Editor/GameProjectGeneration/Private/GameProjectUtils.cpp
+++ b/Engine/Source/Editor/GameProjectGeneration/Private/GameProjectUtils.cpp
@@ -70,7 +70,6 @@
 #include "IAudioExtensionPlugin.h"
 #include "AudioPluginUtilities.h"
 #include "Sound/AudioSettings.h"
-#include "Sound/SoundCueTemplate.h"
 #include "Sound/SoundEffectSubmix.h"
 #include "Sound/SoundEffectSource.h"
 #include "Components/SynthComponent.h"
@@ -524,10 +523,6 @@ FString FNewClassInfo::GetHeaderTemplateFilename() const
 				{
 					return TEXT("SynthComponentClass.h.template");
 				}
-				else if (BaseClass == USoundCueTemplate::StaticClass())
-				{
-					return TEXT("SoundCueTemplateClass.h.template");
-				}
 			}
 			// Some other non-actor, non-component UObject class
 			return TEXT( "UObjectClass.h.template" );
@@ -586,10 +581,6 @@ FString FNewClassInfo::GetSourceTemplateFilename() const
 				{
 					return TEXT("SynthComponentClass.cpp.template");
 				}
-				else if (BaseClass == USoundCueTemplate::StaticClass())
-				{
-					return TEXT("SoundCueTemplateClass.cpp.template");
-				}
 			}
 			// Some other non-actor, non-component UObject class
 			return TEXT( "UObjectClass.cpp.template" );
diff --git a/Engine/Source/Runtime/Engine/Classes/Sound/SoundCueTemplate.h b/Engine/Source/Runtime/Engine/Classes/Sound/SoundCueTemplate.h
deleted file mode 100644
diff --git a/Engine/Source/Runtime/Engine/Private/SoundCueTemplate.cpp b/Engine/Source/Runtime/Engine/Private/SoundCueTemplate.cpp
deleted file mode 100644
