Fix missing #undef LOCTEXT_NAMESPACE causing build warnings.

#rnx

diff --git a/Engine/Plugins/Runtime/SoundCueTemplates/Source/SoundCueTemplatesEditor/Private/SoundCueTemplateFactory.cpp b/Engine/Plugins/Runtime/SoundCueTemplates/Source/SoundCueTemplatesEditor/Private/SoundCueTemplateFactory.cpp
--- a/Engine/Plugins/Runtime/SoundCueTemplates/Source/SoundCueTemplatesEditor/Private/SoundCueTemplateFactory.cpp
+++ b/Engine/Plugins/Runtime/SoundCueTemplates/Source/SoundCueTemplatesEditor/Private/SoundCueTemplateFactory.cpp
@@ -87,3 +87,5 @@ UObject* USoundCueTemplateFactory::FactoryCreateNew(UClass* Class, UObject* InPa
 
 	return nullptr;
 }
+
+#undef LOCTEXT_NAMESPACE // "SoundCueTemplatesEditor"
