SoundCueTemplates:
new template instantiations default to the naming convention "<SoundClassTemplateClassName>_"

This makes it easier for content creators to create and name multiple template assets quickly:
PickaxeReady_Stealth
PickaxeReady_IceCreamCone
PickaxeReady_NovaStar
PickaxeReady_CattusDoggus

#rnx

diff --git a/Engine/Plugins/Runtime/SoundCueTemplates/Source/SoundCueTemplatesEditor/Private/AssetTypeActions_SoundCueTemplate.h b/Engine/Plugins/Runtime/SoundCueTemplates/Source/SoundCueTemplatesEditor/Private/AssetTypeActions_SoundCueTemplate.h
--- a/Engine/Plugins/Runtime/SoundCueTemplates/Source/SoundCueTemplatesEditor/Private/AssetTypeActions_SoundCueTemplate.h
+++ b/Engine/Plugins/Runtime/SoundCueTemplates/Source/SoundCueTemplatesEditor/Private/AssetTypeActions_SoundCueTemplate.h
@@ -12,7 +12,7 @@ class FAssetTypeActions_SoundCueTemplate : public FAssetTypeActions_Base
 public:
 	// IAssetTypeActions Implementation
 	virtual FText GetName() const override { return NSLOCTEXT("AssetTypeActions", "AssetTypeActions_SoundCueTemplate", "Sound Cue Template"); }
-	virtual FColor GetTypeColor() const override { return FColor(255, 255, 255); }
+	virtual FColor GetTypeColor() const override { return FColor::Red; }
 	virtual UClass* GetSupportedClass() const override;
 	virtual bool HasActions(const TArray<UObject*>& InObjects) const override { return true; }
 	virtual void GetActions(const TArray<UObject*>& InObjects, struct FToolMenuSection& Section) override;
diff --git a/Engine/Plugins/Runtime/SoundCueTemplates/Source/SoundCueTemplatesEditor/Private/SoundCueTemplateFactory.cpp b/Engine/Plugins/Runtime/SoundCueTemplates/Source/SoundCueTemplatesEditor/Private/SoundCueTemplateFactory.cpp
--- a/Engine/Plugins/Runtime/SoundCueTemplates/Source/SoundCueTemplatesEditor/Private/SoundCueTemplateFactory.cpp
+++ b/Engine/Plugins/Runtime/SoundCueTemplates/Source/SoundCueTemplatesEditor/Private/SoundCueTemplateFactory.cpp
@@ -88,4 +88,14 @@ UObject* USoundCueTemplateFactory::FactoryCreateNew(UClass* Class, UObject* InPa
 	return nullptr;
 }
 
+FString USoundCueTemplateFactory::GetDefaultNewAssetName() const
+{
+	if (SoundCueTemplateClass)
+	{
+		return SoundCueTemplateClass->GetName() + FString(TEXT("_"));
+	}
+
+	return Super::GetDefaultNewAssetName();
+}
+
 #undef LOCTEXT_NAMESPACE // "SoundCueTemplatesEditor"
diff --git a/Engine/Plugins/Runtime/SoundCueTemplates/Source/SoundCueTemplatesEditor/Public/SoundCueTemplateFactory.h b/Engine/Plugins/Runtime/SoundCueTemplates/Source/SoundCueTemplatesEditor/Public/SoundCueTemplateFactory.h
--- a/Engine/Plugins/Runtime/SoundCueTemplates/Source/SoundCueTemplatesEditor/Public/SoundCueTemplateFactory.h
+++ b/Engine/Plugins/Runtime/SoundCueTemplates/Source/SoundCueTemplatesEditor/Public/SoundCueTemplateFactory.h
@@ -33,5 +33,6 @@ class USoundCueTemplateFactory : public UFactory
 	//~ Begin UFactory Interface
 	virtual bool ConfigureProperties() override;
 	virtual UObject* FactoryCreateNew(UClass* Class, UObject* InParent, FName Name, EObjectFlags Flags, UObject* Context, FFeedbackContext* Warn) override;
+	virtual FString GetDefaultNewAssetName() const override;
 	//~ End UFactory Interface
-};
\ No newline at end of file
+};
