SoundCueTemplates - Adding preview / stop functionality

#rnx

diff --git a/Engine/Plugins/Runtime/SoundCueTemplates/Source/SoundCueTemplatesEditor/Private/AssetTypeActions_SoundCueTemplate.cpp b/Engine/Plugins/Runtime/SoundCueTemplates/Source/SoundCueTemplatesEditor/Private/AssetTypeActions_SoundCueTemplate.cpp
--- a/Engine/Plugins/Runtime/SoundCueTemplates/Source/SoundCueTemplatesEditor/Private/AssetTypeActions_SoundCueTemplate.cpp
+++ b/Engine/Plugins/Runtime/SoundCueTemplates/Source/SoundCueTemplatesEditor/Private/AssetTypeActions_SoundCueTemplate.cpp
@@ -4,6 +4,7 @@
 #include "SoundCueTemplate.h"
 #include "SoundCueTemplateFactory.h"
 
+#include "Components/AudioComponent.h"
 #include "ToolMenuSection.h"
 #include "ContentBrowserModule.h"
 #include "IContentBrowserSingleton.h"
@@ -31,6 +32,62 @@ void FAssetTypeActions_SoundCueTemplate::GetActions(const TArray<UObject*>& InOb
 	);
 }
 
+bool FAssetTypeActions_SoundCueTemplate::AssetsActivatedOverride(const TArray<UObject *>& InObjects, EAssetTypeActivationMethod::Type ActivationType)
+{
+	if (ActivationType == EAssetTypeActivationMethod::Previewed)
+	{
+		USoundBase* TargetSound = NULL;
+
+		for (auto ObjIt = InObjects.CreateConstIterator(); ObjIt; ++ObjIt)
+		{
+			TargetSound = Cast<USoundBase>(*ObjIt);
+			if (TargetSound)
+			{
+				// Only target the first valid sound cue
+				break;
+			}
+		}
+
+		UAudioComponent* PreviewComp = GEditor->GetPreviewAudioComponent();
+		if (PreviewComp && PreviewComp->IsPlaying())
+		{
+			// Already previewing a sound, if it is the target cue then stop it, otherwise play the new one
+			if (!TargetSound || PreviewComp->Sound == TargetSound)
+			{
+				StopSound();
+			}
+			else
+			{
+				PlaySound(TargetSound);
+			}
+		}
+		else
+		{
+			// Not already playing, play the target sound cue if it exists
+			PlaySound(TargetSound);
+		}
+		return true;
+	}
+	return false;
+}
+
+void FAssetTypeActions_SoundCueTemplate::PlaySound(USoundBase* Sound) const
+{
+	if (Sound)
+	{
+		GEditor->PlayPreviewSound(Sound);
+	}
+	else
+	{
+		StopSound();
+	}
+}
+
+void FAssetTypeActions_SoundCueTemplate::StopSound() const
+{
+	GEditor->ResetPreviewAudioComponent();
+}
+
 void FAssetTypeActions_SoundCueTemplate::ExecuteCopyToSoundCue(TArray<TWeakObjectPtr<USoundCueTemplate>> Objects)
 {
 	for (const TWeakObjectPtr<USoundCueTemplate>& Object : Objects)
diff --git a/Engine/Plugins/Runtime/SoundCueTemplates/Source/SoundCueTemplatesEditor/Private/AssetTypeActions_SoundCueTemplate.h b/Engine/Plugins/Runtime/SoundCueTemplates/Source/SoundCueTemplatesEditor/Private/AssetTypeActions_SoundCueTemplate.h
--- a/Engine/Plugins/Runtime/SoundCueTemplates/Source/SoundCueTemplatesEditor/Private/AssetTypeActions_SoundCueTemplate.h
+++ b/Engine/Plugins/Runtime/SoundCueTemplates/Source/SoundCueTemplatesEditor/Private/AssetTypeActions_SoundCueTemplate.h
@@ -16,10 +16,17 @@ public:
 	virtual UClass* GetSupportedClass() const override;
 	virtual bool HasActions(const TArray<UObject*>& InObjects) const override { return true; }
 	virtual void GetActions(const TArray<UObject*>& InObjects, struct FToolMenuSection& Section) override;
+	virtual bool AssetsActivatedOverride(const TArray<UObject *>& InObjects, EAssetTypeActivationMethod::Type ActivationType) override;
 	virtual uint32 GetCategories() override { return EAssetTypeCategories::Sounds; }
 	virtual bool CanFilter() override { return true; }
 
 protected:
+	/** Plays the specified soundcue template */
+	void PlaySound(USoundBase* Sound) const;
+
+	/** Stops any currently playing template */
+	void StopSound() const;
+
 	/** Converts the provided SoundCue Template to a fully-modifiable SoundCue */
 	void ExecuteCopyToSoundCue(TArray<TWeakObjectPtr<USoundCueTemplate>> Objects);
 };
