SoundCueTemplates: Adding support for sound wave right click "create template from these assets"
(swarm review 11281495 - originally in 12.00, moved to 12.10)

#rnx

diff --git a/Engine/Plugins/Runtime/SoundCueTemplates/Source/SoundCueTemplates/Public/SoundCueTemplate.h b/Engine/Plugins/Runtime/SoundCueTemplates/Source/SoundCueTemplates/Public/SoundCueTemplate.h
--- a/Engine/Plugins/Runtime/SoundCueTemplates/Source/SoundCueTemplates/Public/SoundCueTemplate.h
+++ b/Engine/Plugins/Runtime/SoundCueTemplates/Source/SoundCueTemplates/Public/SoundCueTemplate.h
@@ -1,4 +1,5 @@
 // Copyright Epic Games, Inc. All Rights Reserved.
+
 #pragma once
 
 #include "CoreMinimal.h"
@@ -29,6 +30,9 @@ public:
 	/** Rebuilds the graph when user-facing properties are changed */
 	void RebuildGraph(USoundCue& SoundCue) const;
 
+	/** Override to initialize a template with array of sound waves from SoundCueFactory */
+	virtual void AddSoundWaves(TArray<TWeakObjectPtr<USoundWave>>& Waves) {}
+
 protected:
 	/**
 	 * Function to override that uses internal editor-data only properties to re-build node
diff --git a/Engine/Plugins/Runtime/SoundCueTemplates/Source/SoundCueTemplatesEditor/Private/AssetTypeActions_SoundCueTemplate.cpp b/Engine/Plugins/Runtime/SoundCueTemplates/Source/SoundCueTemplatesEditor/Private/AssetTypeActions_SoundCueTemplate.cpp
--- a/Engine/Plugins/Runtime/SoundCueTemplates/Source/SoundCueTemplatesEditor/Private/AssetTypeActions_SoundCueTemplate.cpp
+++ b/Engine/Plugins/Runtime/SoundCueTemplates/Source/SoundCueTemplatesEditor/Private/AssetTypeActions_SoundCueTemplate.cpp
@@ -1,13 +1,17 @@
 // Copyright Epic Games, Inc. All Rights Reserved.
+
 #include "AssetTypeActions_SoundCueTemplate.h"
 
 #include "SoundCueTemplate.h"
 #include "SoundCueTemplateFactory.h"
 
 #include "Components/AudioComponent.h"
+#include "ToolMenus.h"
 #include "ToolMenuSection.h"
 #include "ContentBrowserModule.h"
+#include "ContentBrowserMenuContexts.h"
 #include "IContentBrowserSingleton.h"
+#include "ObjectEditorUtils.h"
 
 #define LOCTEXT_NAMESPACE "AssetTypeActions"
 
@@ -109,4 +113,61 @@ void FAssetTypeActions_SoundCueTemplate::ExecuteCopyToSoundCue(TArray<TWeakObjec
 		}
 	}
 }
+
+void FAssetActionExtender_SoundCueTemplate::RegisterMenus()
+{
+	if (!UToolMenus::IsToolMenuUIEnabled())
+	{
+		return;
+	}
+
+	FToolMenuOwnerScoped MenuOwner("SoundCueTemplate");
+	UToolMenu* Menu = UToolMenus::Get()->ExtendMenu("ContentBrowser.AssetContextMenu.SoundWave");
+	FToolMenuSection& Section = Menu->FindOrAddSection("GetAssetActions");
+
+	Section.AddDynamicEntry("SoundCueTemplateSoundWaveAsset", FNewToolMenuSectionDelegate::CreateLambda([](FToolMenuSection& InSection)
+	{
+		UContentBrowserAssetContextMenuContext* Context = InSection.FindContext<UContentBrowserAssetContextMenuContext>();
+		if (!Context || Context->SelectedObjects.Num() == 0)
+		{
+			return;
+		}
+
+		const TAttribute<FText> Label = LOCTEXT("SoundWave_CreateSoundCueTemplate", "Create SoundCueTemplate");
+		const TAttribute<FText> ToolTip = LOCTEXT("SoundWave_CreateSoundCueTemplateToolTip", "Creates a SoundCueTemplate from the selected sound waves.");
+		const FSlateIcon Icon = FSlateIcon(FEditorStyle::GetStyleSetName(), "ClassIcon.SoundCue");
+		const FToolMenuExecuteAction UIExecuteAction = FToolMenuExecuteAction::CreateStatic(&FAssetActionExtender_SoundCueTemplate::ExecuteCreateSoundCueTemplate);
+
+		InSection.AddMenuEntry("SoundWave_CreateSoundCueTemplate", Label, ToolTip, Icon, UIExecuteAction);
+	}));
+}
+
+void FAssetActionExtender_SoundCueTemplate::ExecuteCreateSoundCueTemplate(const struct FToolMenuContext& MenuContext)
+{
+	UContentBrowserAssetContextMenuContext* Context = MenuContext.FindContext<UContentBrowserAssetContextMenuContext>();
+	if (!Context || Context->SelectedObjects.Num() == 0)
+	{
+		return;
+	}
+
+	USoundWave* SoundWave = Cast<USoundWave>(Context->SelectedObjects[0]);
+	if (SoundWave == nullptr)
+	{
+		return;
+	}
+
+	FString PackagePath;
+	FString Name;
+
+	FAssetToolsModule& AssetToolsModule = FModuleManager::Get().LoadModuleChecked<FAssetToolsModule>("AssetTools");
+	AssetToolsModule.Get().CreateUniqueAssetName(SoundWave->GetOutermost()->GetName(), TEXT(""), PackagePath, Name);
+
+	USoundCueTemplateFactory* Factory = NewObject<USoundCueTemplateFactory>();
+	Factory->ConfigureProperties();
+
+	Name = Factory->GetDefaultNewAssetName();
+	Factory->SoundWaves = FObjectEditorUtils::GetTypedWeakObjectPtrs<USoundWave>(Context->GetSelectedObjects());
+
+	AssetToolsModule.Get().CreateAsset(Name, FPackageName::GetLongPackagePath(PackagePath), USoundCueTemplate::StaticClass(), Factory);
+}
 #undef LOCTEXT_NAMESPACE
diff --git a/Engine/Plugins/Runtime/SoundCueTemplates/Source/SoundCueTemplatesEditor/Private/AssetTypeActions_SoundCueTemplate.h b/Engine/Plugins/Runtime/SoundCueTemplates/Source/SoundCueTemplatesEditor/Private/AssetTypeActions_SoundCueTemplate.h
--- a/Engine/Plugins/Runtime/SoundCueTemplates/Source/SoundCueTemplatesEditor/Private/AssetTypeActions_SoundCueTemplate.h
+++ b/Engine/Plugins/Runtime/SoundCueTemplates/Source/SoundCueTemplatesEditor/Private/AssetTypeActions_SoundCueTemplate.h
@@ -1,4 +1,5 @@
 // Copyright Epic Games, Inc. All Rights Reserved.
+
 #pragma once
 
 #include "CoreMinimal.h"
@@ -30,3 +31,10 @@ protected:
 	/** Converts the provided SoundCue Template to a fully-modifiable SoundCue */
 	void ExecuteCopyToSoundCue(TArray<TWeakObjectPtr<USoundCueTemplate>> Objects);
 };
+
+class FAssetActionExtender_SoundCueTemplate
+{
+public:
+	static void RegisterMenus();
+	static void ExecuteCreateSoundCueTemplate(const struct FToolMenuContext& MenuContext);
+};
\ No newline at end of file
diff --git a/Engine/Plugins/Runtime/SoundCueTemplates/Source/SoundCueTemplatesEditor/Private/SoundCueTemplateFactory.cpp b/Engine/Plugins/Runtime/SoundCueTemplates/Source/SoundCueTemplatesEditor/Private/SoundCueTemplateFactory.cpp
--- a/Engine/Plugins/Runtime/SoundCueTemplates/Source/SoundCueTemplatesEditor/Private/SoundCueTemplateFactory.cpp
+++ b/Engine/Plugins/Runtime/SoundCueTemplates/Source/SoundCueTemplatesEditor/Private/SoundCueTemplateFactory.cpp
@@ -1,4 +1,5 @@
 // Copyright Epic Games, Inc. All Rights Reserved.
+
 #include "SoundCueTemplateFactory.h"
 #include "SoundCueTemplateClassFilter.h"
 
@@ -54,7 +55,7 @@ bool USoundCueTemplateFactory::ConfigureProperties()
 	// Load the classviewer module to display a class picker
 	FClassViewerModule& ClassViewerModule = FModuleManager::LoadModuleChecked<FClassViewerModule>("ClassViewer");
 
-	TSharedPtr<FSoundCueTemplateAssetParentFilter> Filter = MakeShareable(new FSoundCueTemplateAssetParentFilter);
+	TSharedPtr<FSoundCueTemplateAssetParentFilter> Filter = MakeShared<FSoundCueTemplateAssetParentFilter>();
 	Filter->DisallowedClassFlags = CLASS_Abstract | CLASS_Deprecated | CLASS_NewerVersionExists | CLASS_HideDropDown;
 	Filter->AllowedChildrenOfClasses.Add(USoundCueTemplate::StaticClass());
 
@@ -80,6 +81,11 @@ UObject* USoundCueTemplateFactory::FactoryCreateNew(UClass* Class, UObject* InPa
 	{
 		if (USoundCueTemplate* NewSoundCueTemplate = NewObject<USoundCueTemplate>(InParent, SoundCueTemplateClass, Name, Flags))
 		{
+			if (SoundWaves.Num())
+			{
+				NewSoundCueTemplate->AddSoundWaves(SoundWaves);
+			}
+
 			NewSoundCueTemplate->RebuildGraph(*NewSoundCueTemplate);
 			return NewSoundCueTemplate;
 		}
diff --git a/Engine/Plugins/Runtime/SoundCueTemplates/Source/SoundCueTemplatesEditor/Private/SoundCueTemplatesEditorModule.cpp b/Engine/Plugins/Runtime/SoundCueTemplates/Source/SoundCueTemplatesEditor/Private/SoundCueTemplatesEditorModule.cpp
--- a/Engine/Plugins/Runtime/SoundCueTemplates/Source/SoundCueTemplatesEditor/Private/SoundCueTemplatesEditorModule.cpp
+++ b/Engine/Plugins/Runtime/SoundCueTemplates/Source/SoundCueTemplatesEditor/Private/SoundCueTemplatesEditorModule.cpp
@@ -4,16 +4,20 @@
 
 #include "Modules/ModuleManager.h"
 #include "AssetToolsModule.h"
+#include "ToolMenus.h"
 
 #include "AssetTypeActions_SoundCueTemplate.h"
 
 void FSoundCueTemplatesEditorModule::StartupModule()
 {
 	RegisterAssetActions();
+	UToolMenus::RegisterStartupCallback(FSimpleMulticastDelegate::FDelegate::CreateRaw(this, &FSoundCueTemplatesEditorModule::RegisterMenus));
 }
 
 void FSoundCueTemplatesEditorModule::ShutdownModule()
 {
+	UToolMenus::UnRegisterStartupCallback(this);
+	UToolMenus::UnregisterOwner("SoundCueTemplates");
 }
 
 void FSoundCueTemplatesEditorModule::RegisterAssetActions()
@@ -24,4 +28,9 @@ void FSoundCueTemplatesEditorModule::RegisterAssetActions()
 	AssetTools.RegisterAssetTypeActions(MakeShared<FAssetTypeActions_SoundCueTemplate>());
 }
 
+void FSoundCueTemplatesEditorModule::RegisterMenus()
+{
+	FAssetActionExtender_SoundCueTemplate::RegisterMenus();
+}
+
 IMPLEMENT_MODULE(FSoundCueTemplatesEditorModule, SoundCueTemplatesEditor);
diff --git a/Engine/Plugins/Runtime/SoundCueTemplates/Source/SoundCueTemplatesEditor/Public/SoundCueTemplateFactory.h b/Engine/Plugins/Runtime/SoundCueTemplates/Source/SoundCueTemplatesEditor/Public/SoundCueTemplateFactory.h
--- a/Engine/Plugins/Runtime/SoundCueTemplates/Source/SoundCueTemplatesEditor/Public/SoundCueTemplateFactory.h
+++ b/Engine/Plugins/Runtime/SoundCueTemplates/Source/SoundCueTemplatesEditor/Public/SoundCueTemplateFactory.h
@@ -1,4 +1,5 @@
 // Copyright Epic Games, Inc. All Rights Reserved.
+
 #pragma once
 
 #include "CoreMinimal.h"
@@ -35,4 +36,6 @@ class USoundCueTemplateFactory : public UFactory
 	virtual UObject* FactoryCreateNew(UClass* Class, UObject* InParent, FName Name, EObjectFlags Flags, UObject* Context, FFeedbackContext* Warn) override;
 	virtual FString GetDefaultNewAssetName() const override;
 	//~ End UFactory Interface
+
+	TArray<TWeakObjectPtr<USoundWave>> SoundWaves;
 };
diff --git a/Engine/Plugins/Runtime/SoundCueTemplates/Source/SoundCueTemplatesEditor/Public/SoundCueTemplatesEditorModule.h b/Engine/Plugins/Runtime/SoundCueTemplates/Source/SoundCueTemplatesEditor/Public/SoundCueTemplatesEditorModule.h
--- a/Engine/Plugins/Runtime/SoundCueTemplates/Source/SoundCueTemplatesEditor/Public/SoundCueTemplatesEditorModule.h
+++ b/Engine/Plugins/Runtime/SoundCueTemplates/Source/SoundCueTemplatesEditor/Public/SoundCueTemplatesEditorModule.h
@@ -3,7 +3,6 @@
 #pragma once
 
 #include "CoreMinimal.h"
-#include "Containers/Set.h"
 #include "Modules/ModuleInterface.h"
 
 SOUNDCUETEMPLATESEDITOR_API DECLARE_LOG_CATEGORY_EXTERN(SoundCueTemplatesEditor, Log, All);
@@ -17,4 +16,5 @@ public:
 
 private:
 	void RegisterAssetActions();
+	void RegisterMenus();
 };
diff --git a/Engine/Plugins/Runtime/SoundCueTemplates/Source/SoundCueTemplatesEditor/SoundCueTemplatesEditor.Build.cs b/Engine/Plugins/Runtime/SoundCueTemplates/Source/SoundCueTemplatesEditor/SoundCueTemplatesEditor.Build.cs
--- a/Engine/Plugins/Runtime/SoundCueTemplates/Source/SoundCueTemplatesEditor/SoundCueTemplatesEditor.Build.cs
+++ b/Engine/Plugins/Runtime/SoundCueTemplates/Source/SoundCueTemplatesEditor/SoundCueTemplatesEditor.Build.cs
@@ -25,7 +25,8 @@ public class SoundCueTemplatesEditor : ModuleRules
                 "Engine",
                 "UnrealEd",
                 "Slate",
-                "SlateCore"
+                "SlateCore",
+                "ContentBrowser"
             }
         );
 
