SoundCueTemplate: Fixing asset action to allow rename before asset is finalized.

diff --git a/Engine/Plugins/Runtime/SoundCueTemplates/Source/SoundCueTemplatesEditor/Private/AssetTypeActions_SoundCueTemplate.cpp b/Engine/Plugins/Runtime/SoundCueTemplates/Source/SoundCueTemplatesEditor/Private/AssetTypeActions_SoundCueTemplate.cpp
--- a/Engine/Plugins/Runtime/SoundCueTemplates/Source/SoundCueTemplatesEditor/Private/AssetTypeActions_SoundCueTemplate.cpp
+++ b/Engine/Plugins/Runtime/SoundCueTemplates/Source/SoundCueTemplatesEditor/Private/AssetTypeActions_SoundCueTemplate.cpp
@@ -168,6 +168,7 @@ void FAssetActionExtender_SoundCueTemplate::ExecuteCreateSoundCueTemplate(const
 	Name = Factory->GetDefaultNewAssetName();
 	Factory->SoundWaves = FObjectEditorUtils::GetTypedWeakObjectPtrs<USoundWave>(Context->GetSelectedObjects());
 
-	AssetToolsModule.Get().CreateAsset(Name, FPackageName::GetLongPackagePath(PackagePath), USoundCueTemplate::StaticClass(), Factory);
+	FContentBrowserModule& ContentBrowserModule = FModuleManager::Get().LoadModuleChecked<FContentBrowserModule>("ContentBrowser");
+	ContentBrowserModule.Get().CreateNewAsset(Name, FPackageName::GetLongPackagePath(PackagePath), USoundCueTemplate::StaticClass(), Factory);
 }
 #undef LOCTEXT_NAMESPACE
