Edigrate fixes from 4.26: SoundCueTemplate update & Modulation Mix crash fix

diff --git a/Engine/Plugins/Runtime/AudioModulation/AudioModulation.uplugin b/Engine/Plugins/Runtime/AudioModulation/AudioModulation.uplugin
--- a/Engine/Plugins/Runtime/AudioModulation/AudioModulation.uplugin
+++ b/Engine/Plugins/Runtime/AudioModulation/AudioModulation.uplugin
@@ -3,7 +3,7 @@
 	"Version": 1,
 	"VersionName": "2.0",
 	"FriendlyName": "Audio Modulation",
-	"Description": "Enables modulation of control parameters in the Unreal Audio Engine.",
+	"Description": "Default implementation of Audio Modulation in the Unreal Audio Engine.",
 	"Category": "Audio",
 	"CreatedBy": "Epic Games, Inc.",
 	"CreatedByURL": "http://epicgames.com",
diff --git a/Engine/Plugins/Runtime/AudioModulation/Source/AudioModulation/Private/SoundControlBusMixProxy.cpp b/Engine/Plugins/Runtime/AudioModulation/Source/AudioModulation/Private/SoundControlBusMixProxy.cpp
--- a/Engine/Plugins/Runtime/AudioModulation/Source/AudioModulation/Private/SoundControlBusMixProxy.cpp
+++ b/Engine/Plugins/Runtime/AudioModulation/Source/AudioModulation/Private/SoundControlBusMixProxy.cpp
@@ -22,11 +22,19 @@ namespace AudioModulation
 	FModulatorBusMixStageSettings::FModulatorBusMixStageSettings(const FSoundControlBusMixStage& InStage, Audio::FDeviceId InDeviceId)
 		: TModulatorBase<FBusId>(InStage.Bus->GetName(), InStage.Bus->GetUniqueID())
 		, Address(InStage.Bus->Address)
-		, ParamClassId(InStage.Bus->Parameter->GetClass()->GetUniqueID())
-		, ParamId(InStage.Bus->Parameter->GetUniqueID())
+		, ParamClassId(INDEX_NONE)
+		, ParamId(INDEX_NONE)
 		, Value(InStage.Value)
 		, BusSettings(FControlBusSettings(*InStage.Bus, InDeviceId))
 	{
+		if (USoundModulationParameter* Parameter = InStage.Bus->Parameter)
+		{
+			ParamId = Parameter->GetUniqueID();
+
+			UClass* Class = Parameter->GetClass();
+			check(Class);
+			ParamClassId = Class->GetUniqueID();
+		}
 	}
 
 	FModulatorBusMixStageProxy::FModulatorBusMixStageProxy(const FModulatorBusMixStageSettings& InSettings, FAudioModulationSystem& OutModSystem)
diff --git a/Engine/Plugins/Runtime/SoundCueTemplates/SoundCueTemplates.uplugin b/Engine/Plugins/Runtime/SoundCueTemplates/SoundCueTemplates.uplugin
--- a/Engine/Plugins/Runtime/SoundCueTemplates/SoundCueTemplates.uplugin
+++ b/Engine/Plugins/Runtime/SoundCueTemplates/SoundCueTemplates.uplugin
@@ -10,7 +10,7 @@
 	"DocsURL": "",
 	"MarketplaceURL": "",
 	"SupportURL": "",
-	"CanContainContent": false,
+	"CanContainContent": true,
 	"IsBetaVersion": true,
 	"Installed": false,
 	"Modules": [
diff --git a/Engine/Plugins/Runtime/SoundCueTemplates/Source/SoundCueTemplatesEditor/Private/SoundCueTemplateFactory.cpp b/Engine/Plugins/Runtime/SoundCueTemplates/Source/SoundCueTemplatesEditor/Private/SoundCueTemplateFactory.cpp
--- a/Engine/Plugins/Runtime/SoundCueTemplates/Source/SoundCueTemplatesEditor/Private/SoundCueTemplateFactory.cpp
+++ b/Engine/Plugins/Runtime/SoundCueTemplates/Source/SoundCueTemplatesEditor/Private/SoundCueTemplateFactory.cpp
@@ -104,4 +104,10 @@ FString USoundCueTemplateFactory::GetDefaultNewAssetName() const
 	return Super::GetDefaultNewAssetName();
 }
 
+USoundCueTemplateClassTemplate::USoundCueTemplateClassTemplate(const FObjectInitializer& ObjectInitializer)
+	: UPluginClassTemplate(ObjectInitializer)
+{
+	PluginName = TEXT("SoundCueTemplates");
+	SetGeneratedBaseClass(USoundCueTemplate::StaticClass());
+}
 #undef LOCTEXT_NAMESPACE // "SoundCueTemplatesEditor"
diff --git a/Engine/Plugins/Runtime/SoundCueTemplates/Source/SoundCueTemplatesEditor/Public/SoundCueTemplateFactory.h b/Engine/Plugins/Runtime/SoundCueTemplates/Source/SoundCueTemplatesEditor/Public/SoundCueTemplateFactory.h
--- a/Engine/Plugins/Runtime/SoundCueTemplates/Source/SoundCueTemplatesEditor/Public/SoundCueTemplateFactory.h
+++ b/Engine/Plugins/Runtime/SoundCueTemplates/Source/SoundCueTemplatesEditor/Public/SoundCueTemplateFactory.h
@@ -2,10 +2,11 @@
 
 #pragma once
 
+#include "ClassTemplateEditorSubsystem.h"
 #include "CoreMinimal.h"
-#include "UObject/ObjectMacros.h"
-#include "SoundCueTemplate.h"
 #include "Factories/Factory.h"
+#include "SoundCueTemplate.h"
+#include "UObject/ObjectMacros.h"
 
 #include "SoundCueTemplateFactory.generated.h"
 
@@ -39,3 +40,9 @@ class USoundCueTemplateFactory : public UFactory
 
 	TArray<TWeakObjectPtr<USoundWave>> SoundWaves;
 };
+
+UCLASS()
+class USoundCueTemplateClassTemplate : public UPluginClassTemplate
+{
+	GENERATED_UCLASS_BODY()
+};
diff --git a/Engine/Plugins/Runtime/SoundCueTemplates/Source/SoundCueTemplatesEditor/SoundCueTemplatesEditor.Build.cs b/Engine/Plugins/Runtime/SoundCueTemplates/Source/SoundCueTemplatesEditor/SoundCueTemplatesEditor.Build.cs
--- a/Engine/Plugins/Runtime/SoundCueTemplates/Source/SoundCueTemplatesEditor/SoundCueTemplatesEditor.Build.cs
+++ b/Engine/Plugins/Runtime/SoundCueTemplates/Source/SoundCueTemplatesEditor/SoundCueTemplatesEditor.Build.cs
@@ -6,22 +6,23 @@ public class SoundCueTemplatesEditor : ModuleRules
 {
 	public SoundCueTemplatesEditor(ReadOnlyTargetRules Target) : base(Target)
 	{
-        PCHUsage = ModuleRules.PCHUsageMode.UseExplicitOrSharedPCHs;
+		PCHUsage = ModuleRules.PCHUsageMode.UseExplicitOrSharedPCHs;
 
-        PublicDependencyModuleNames.AddRange(
-            new string[]
-            {
-                "Core",
-                "SoundCueTemplates",
-                "ToolMenus"
-            }
-        );
+		PublicDependencyModuleNames.AddRange(
+			new string[]
+			{
+				"Core",
+				"SoundCueTemplates",
+				"ToolMenus"
+			}
+		);
 
         PrivateDependencyModuleNames.AddRange(
             new string[]
             {
                 "CoreUObject",
 				"EditorFramework",
+				"GameProjectGeneration",
                 "EditorStyle",
                 "Engine",
                 "UnrealEd",
@@ -31,10 +32,10 @@ public class SoundCueTemplatesEditor : ModuleRules
             }
         );
 
-        PrivateIncludePathModuleNames.AddRange(
-            new string[] {
-                "AssetTools"
-            }
-        );
-    }
+		PrivateIncludePathModuleNames.AddRange(
+			new string[] {
+				"AssetTools"
+			}
+		);
+	}
 }
