Providing more customization w/ default names for new sound cue template assets

diff --git a/Engine/Plugins/Runtime/SoundCueTemplates/Source/SoundCueTemplates/Public/SoundCueTemplate.h b/Engine/Plugins/Runtime/SoundCueTemplates/Source/SoundCueTemplates/Public/SoundCueTemplate.h
--- a/Engine/Plugins/Runtime/SoundCueTemplates/Source/SoundCueTemplates/Public/SoundCueTemplate.h
+++ b/Engine/Plugins/Runtime/SoundCueTemplates/Source/SoundCueTemplates/Public/SoundCueTemplate.h
@@ -33,6 +33,9 @@ public:
 	/** Override to initialize a template with array of sound waves from SoundCueFactory */
 	virtual void AddSoundWaves(TArray<TWeakObjectPtr<USoundWave>>& Waves) {}
 
+	/** Override to provide a default new asset name prefix from SoundCueTemplateFactory */
+	virtual FString GenerateDefaultNewAssetName(const TArray<TWeakObjectPtr<USoundWave>>& Waves) const { return FString(); }
+
 protected:
 	/**
 	 * Function to override that uses internal editor-data only properties to re-build node
diff --git a/Engine/Plugins/Runtime/SoundCueTemplates/Source/SoundCueTemplatesEditor/Private/SoundCueTemplateFactory.cpp b/Engine/Plugins/Runtime/SoundCueTemplates/Source/SoundCueTemplatesEditor/Private/SoundCueTemplateFactory.cpp
--- a/Engine/Plugins/Runtime/SoundCueTemplates/Source/SoundCueTemplatesEditor/Private/SoundCueTemplateFactory.cpp
+++ b/Engine/Plugins/Runtime/SoundCueTemplates/Source/SoundCueTemplatesEditor/Private/SoundCueTemplateFactory.cpp
@@ -98,7 +98,18 @@ FString USoundCueTemplateFactory::GetDefaultNewAssetName() const
 {
 	if (SoundCueTemplateClass)
 	{
-		return SoundCueTemplateClass->GetName() + FString(TEXT("_"));
+		FString TemplatePrefix;
+		if (const USoundCueTemplate* TemplateCDO = SoundCueTemplateClass.GetDefaultObject())
+		{
+			TemplatePrefix = TemplateCDO->GenerateDefaultNewAssetName(SoundWaves);
+		}
+
+		if (TemplatePrefix.IsEmpty())
+		{
+			SoundCueTemplateClass->GetName(TemplatePrefix);
+		}
+
+		return TemplatePrefix + FString(TEXT("_"));
 	}
 
 	return Super::GetDefaultNewAssetName();
