Moving node height / width customizations to cvars for SoundCueTemplates

#tests: local testing

diff --git a/Engine/Plugins/Runtime/SoundCueTemplates/Source/SoundCueTemplates/Private/SoundCueTemplate.cpp b/Engine/Plugins/Runtime/SoundCueTemplates/Source/SoundCueTemplates/Private/SoundCueTemplate.cpp
--- a/Engine/Plugins/Runtime/SoundCueTemplates/Source/SoundCueTemplates/Private/SoundCueTemplate.cpp
+++ b/Engine/Plugins/Runtime/SoundCueTemplates/Source/SoundCueTemplates/Private/SoundCueTemplate.cpp
@@ -2,6 +2,37 @@
 
 #include "SoundCueTemplate.h"
 
+namespace SoundCueTemplateConsoleVariables
+{
+	float NodeWidthOffset = 500.f;
+	FAutoConsoleVariableRef CVarNodeWidthOffset(
+		TEXT("au.SoundCueTemplate.NodeWidthOffset"),
+		NodeWidthOffset,
+		TEXT("Sets the width between nodes in SoundCues generated from SoundCueTemplates"),
+		ECVF_Default);
+
+	float InitialWidthOffset = 300.f;
+	FAutoConsoleVariableRef CVarInitialWidthOffset(
+		TEXT("au.SoundCueTemplate.InitialWidthOffset"),
+		InitialWidthOffset,
+		TEXT("Sets the initial width offset of first node in a SoundCue generated from a SoundCueTemplate"),
+		ECVF_Default);
+
+	float NodeHeightOffset = 200.f;
+	FAutoConsoleVariableRef CVarNodeHeightOffset(
+		TEXT("au.SoundCueTemplate.NodeHeightOffset"),
+		NodeHeightOffset,
+		TEXT("Sets the height between nodes in SoundCues generated from SoundCueTemplates"),
+		ECVF_Default);
+
+	float InitialHeightOffset = 0.f;
+	FAutoConsoleVariableRef CVarInitialHeightOffset(
+		TEXT("au.SoundCueTemplate.InitialHeightOffset"),
+		InitialHeightOffset,
+		TEXT("Sets the initial height offset of first node in a SoundCue generated from a SoundCueTemplate"),
+		ECVF_Default);
+} // namespace SoundCueTemplateConsoleVariables
+
 USoundCueTemplate::USoundCueTemplate(const FObjectInitializer& ObjectInitializer)
 	: Super(ObjectInitializer)
 {
@@ -28,4 +59,25 @@ void USoundCueTemplate::RebuildGraph(USoundCue& SoundCue) const
 	SoundCue.PostEditChange();
 	SoundCue.MarkPackageDirty();
 }
+
+float USoundCueTemplate::GetNodeWidthOffset()
+{
+	return SoundCueTemplateConsoleVariables::NodeWidthOffset;
+}
+
+float USoundCueTemplate::GetInitialWidthOffset()
+{
+	return SoundCueTemplateConsoleVariables::InitialWidthOffset;
+}
+
+float USoundCueTemplate::GetNodeHeightOffset()
+{
+	return SoundCueTemplateConsoleVariables::NodeHeightOffset;
+}
+
+float USoundCueTemplate::GetInitialHeightOffset()
+{
+	return SoundCueTemplateConsoleVariables::InitialHeightOffset;
+}
+
 #endif // WITH_EDITOR
diff --git a/Engine/Plugins/Runtime/SoundCueTemplates/Source/SoundCueTemplates/Public/SoundCueTemplate.h b/Engine/Plugins/Runtime/SoundCueTemplates/Source/SoundCueTemplates/Public/SoundCueTemplate.h
--- a/Engine/Plugins/Runtime/SoundCueTemplates/Source/SoundCueTemplates/Public/SoundCueTemplate.h
+++ b/Engine/Plugins/Runtime/SoundCueTemplates/Source/SoundCueTemplates/Public/SoundCueTemplate.h
@@ -47,34 +47,22 @@ protected:
 	/**
 	 * Node width offset between nodes in the SoundCueEditor.
 	 */
-	static float GetNodeWidthOffset()
-	{
-		return 500.0f;
-	}
+	static float GetNodeWidthOffset();
 
 	/**
 	 * Initial width offset of first node in the SoundCueEditor.
 	 */
-	static float GetInitialWidthOffset()
-	{
-		return 300.0f;
-	}
+	static float GetInitialWidthOffset();
 
 	/**
 	 * Node height offset between nodes in the SoundCueEditor.
 	 */
-	static float GetNodeHeightOffset()
-	{
-		return 200.0f;
-	}
+	static float GetNodeHeightOffset();
 
 	/**
 	 * Initial height offset of first node in the SoundCueEditor.
 	 */
-	static float GetInitialHeightOffset()
-	{
-		return 0.0f;
-	}
+	static float GetInitialHeightOffset();
 
 	/**
 	 * Sets the provided node's position to the location on the template's grid
