Adding hook to allow loading of default template settings after SoundCueTemplate creation

#rnx

diff --git a/Engine/Plugins/Runtime/SoundCueTemplates/Source/SoundCueTemplates/Public/SoundCueTemplate.h b/Engine/Plugins/Runtime/SoundCueTemplates/Source/SoundCueTemplates/Public/SoundCueTemplate.h
--- a/Engine/Plugins/Runtime/SoundCueTemplates/Source/SoundCueTemplates/Public/SoundCueTemplate.h
+++ b/Engine/Plugins/Runtime/SoundCueTemplates/Source/SoundCueTemplates/Public/SoundCueTemplate.h
@@ -33,6 +33,9 @@ public:
 	/** Override to initialize a template with array of sound waves from SoundCueFactory */
 	virtual void AddSoundWaves(TArray<TWeakObjectPtr<USoundWave>>& Waves) {}
 
+	/** Override to load template default settings from SoundCueTemplateFactory */
+	virtual void LoadTemplateDefaultSettings() {}
+
 	/** Override to provide a default new asset name prefix from SoundCueTemplateFactory */
 	virtual FString GenerateDefaultNewAssetName(const TArray<TWeakObjectPtr<USoundWave>>& Waves) const { return FString(); }
 
diff --git a/Engine/Plugins/Runtime/SoundCueTemplates/Source/SoundCueTemplatesEditor/Private/SoundCueTemplateFactory.cpp b/Engine/Plugins/Runtime/SoundCueTemplates/Source/SoundCueTemplatesEditor/Private/SoundCueTemplateFactory.cpp
--- a/Engine/Plugins/Runtime/SoundCueTemplates/Source/SoundCueTemplatesEditor/Private/SoundCueTemplateFactory.cpp
+++ b/Engine/Plugins/Runtime/SoundCueTemplates/Source/SoundCueTemplatesEditor/Private/SoundCueTemplateFactory.cpp
@@ -81,6 +81,8 @@ UObject* USoundCueTemplateFactory::FactoryCreateNew(UClass* Class, UObject* InPa
 	{
 		if (USoundCueTemplate* NewSoundCueTemplate = NewObject<USoundCueTemplate>(InParent, SoundCueTemplateClass, Name, Flags))
 		{
+			NewSoundCueTemplate->LoadTemplateDefaultSettings();
+
 			if (SoundWaves.Num())
 			{
 				NewSoundCueTemplate->AddSoundWaves(SoundWaves);
