SoundCueTemplate - exposing method for adding SoundWaves to template to BP

diff --git a/Engine/Plugins/Runtime/SoundCueTemplates/Source/SoundCueTemplates/Private/SoundCueTemplate.cpp b/Engine/Plugins/Runtime/SoundCueTemplates/Source/SoundCueTemplates/Private/SoundCueTemplate.cpp
--- a/Engine/Plugins/Runtime/SoundCueTemplates/Source/SoundCueTemplates/Private/SoundCueTemplate.cpp
+++ b/Engine/Plugins/Runtime/SoundCueTemplates/Source/SoundCueTemplates/Private/SoundCueTemplate.cpp
@@ -1,6 +1,7 @@
 // Copyright Epic Games, Inc. All Rights Reserved.
 
 #include "SoundCueTemplate.h"
+#include "ObjectEditorUtils.h"
 
 namespace SoundCueTemplateConsoleVariables
 {
@@ -60,6 +61,12 @@ void USoundCueTemplate::RebuildGraph(USoundCue& SoundCue) const
 	SoundCue.MarkPackageDirty();
 }
 
+void USoundCueTemplate::AddSoundWavesToTemplate(TArray<UObject*>& SelectedObjects)
+{
+	TArray<TWeakObjectPtr<USoundWave>> Waves(FObjectEditorUtils::GetTypedWeakObjectPtrs<USoundWave>(SelectedObjects));
+	AddSoundWaves(Waves);
+}
+
 float USoundCueTemplate::GetNodeWidthOffset()
 {
 	return SoundCueTemplateConsoleVariables::NodeWidthOffset;
diff --git a/Engine/Plugins/Runtime/SoundCueTemplates/Source/SoundCueTemplates/Public/SoundCueTemplate.h b/Engine/Plugins/Runtime/SoundCueTemplates/Source/SoundCueTemplates/Public/SoundCueTemplate.h
--- a/Engine/Plugins/Runtime/SoundCueTemplates/Source/SoundCueTemplates/Public/SoundCueTemplate.h
+++ b/Engine/Plugins/Runtime/SoundCueTemplates/Source/SoundCueTemplates/Public/SoundCueTemplate.h
@@ -39,6 +39,9 @@ public:
 	/** Override to provide a default new asset name prefix from SoundCueTemplateFactory */
 	virtual FString GenerateDefaultNewAssetName(const TArray<TWeakObjectPtr<USoundWave>>& Waves) const { return FString(); }
 
+	UFUNCTION(BlueprintCallable, Category = Config)
+	void AddSoundWavesToTemplate(TArray<UObject*>& SelectedObjects);
+
 protected:
 	/**
 	 * Function to override that uses internal editor-data only properties to re-build node
