Updated SoundCueTemplates to not compile with optimizations off so it matches other modules that use Shared PCHs. This speeds up compile times and removes the need to have a 2.03 GB shared Engine PCH.

diff --git a/Engine/Plugins/Runtime/SoundCueTemplates/Source/SoundCueTemplates/SoundCueTemplates.Build.cs b/Engine/Plugins/Runtime/SoundCueTemplates/Source/SoundCueTemplates/SoundCueTemplates.Build.cs
--- a/Engine/Plugins/Runtime/SoundCueTemplates/Source/SoundCueTemplates/SoundCueTemplates.Build.cs
+++ b/Engine/Plugins/Runtime/SoundCueTemplates/Source/SoundCueTemplates/SoundCueTemplates.Build.cs
@@ -6,7 +6,6 @@ public class SoundCueTemplates : ModuleRules
 {
     public SoundCueTemplates(ReadOnlyTargetRules Target) : base(Target)
     {
-        OptimizeCode = CodeOptimization.Never;
         PCHUsage = ModuleRules.PCHUsageMode.UseExplicitOrSharedPCHs;
 
         PublicDependencyModuleNames.AddRange(
