SoundCueTemplate - Add const to AddSoundWavesToTemplate param so that it shows up as an input BP instead of an output

diff --git a/Engine/Plugins/Runtime/SoundCueTemplates/Source/SoundCueTemplates/Private/SoundCueTemplate.cpp b/Engine/Plugins/Runtime/SoundCueTemplates/Source/SoundCueTemplates/Private/SoundCueTemplate.cpp
--- a/Engine/Plugins/Runtime/SoundCueTemplates/Source/SoundCueTemplates/Private/SoundCueTemplate.cpp
+++ b/Engine/Plugins/Runtime/SoundCueTemplates/Source/SoundCueTemplates/Private/SoundCueTemplate.cpp
@@ -63,7 +63,7 @@ void USoundCueTemplate::RebuildGraph(USoundCue& SoundCue) const
 	SoundCue.MarkPackageDirty();
 }
 
-void USoundCueTemplate::AddSoundWavesToTemplate(TArray<UObject*>& SelectedObjects)
+void USoundCueTemplate::AddSoundWavesToTemplate(const TArray<UObject*>& SelectedObjects)
 {
 	TArray<TWeakObjectPtr<USoundWave>> Waves(FObjectEditorUtils::GetTypedWeakObjectPtrs<USoundWave>(SelectedObjects));
 	AddSoundWaves(Waves);
diff --git a/Engine/Plugins/Runtime/SoundCueTemplates/Source/SoundCueTemplates/Public/SoundCueTemplate.h b/Engine/Plugins/Runtime/SoundCueTemplates/Source/SoundCueTemplates/Public/SoundCueTemplate.h
--- a/Engine/Plugins/Runtime/SoundCueTemplates/Source/SoundCueTemplates/Public/SoundCueTemplate.h
+++ b/Engine/Plugins/Runtime/SoundCueTemplates/Source/SoundCueTemplates/Public/SoundCueTemplate.h
@@ -40,7 +40,7 @@ public:
 	virtual FString GenerateDefaultNewAssetName(const TArray<TWeakObjectPtr<USoundWave>>& Waves) const { return FString(); }
 
 	UFUNCTION(BlueprintCallable, Category = Config)
-	void AddSoundWavesToTemplate(TArray<UObject*>& SelectedObjects);
+	void AddSoundWavesToTemplate(const TArray<UObject*>& SelectedObjects);
 
 protected:
 	/**
