Expose some properties on our SoundCue and MetaSound factories so editor scripting can take advantage of them

diff --git a/Engine/Plugins/Runtime/Metasound/Source/MetasoundEditor/Public/MetasoundFactory.h b/Engine/Plugins/Runtime/Metasound/Source/MetasoundEditor/Public/MetasoundFactory.h
--- a/Engine/Plugins/Runtime/Metasound/Source/MetasoundEditor/Public/MetasoundFactory.h
+++ b/Engine/Plugins/Runtime/Metasound/Source/MetasoundEditor/Public/MetasoundFactory.h
@@ -18,7 +18,7 @@ class UMetaSoundBaseFactory : public UFactory
 
 	// Set to initialize MetaSound type as Preset,
 	// using the provided MetaSound as a Reference
-	UPROPERTY(Transient)
+	UPROPERTY(EditAnywhere, Transient, Category = Factory)
 	TObjectPtr<UObject> ReferencedMetaSoundObject;
 };
 
diff --git a/Engine/Plugins/Runtime/SoundCueTemplates/Source/SoundCueTemplatesEditor/Public/SoundCueTemplateFactory.h b/Engine/Plugins/Runtime/SoundCueTemplates/Source/SoundCueTemplatesEditor/Public/SoundCueTemplateFactory.h
--- a/Engine/Plugins/Runtime/SoundCueTemplates/Source/SoundCueTemplatesEditor/Public/SoundCueTemplateFactory.h
+++ b/Engine/Plugins/Runtime/SoundCueTemplates/Source/SoundCueTemplatesEditor/Public/SoundCueTemplateFactory.h
@@ -32,13 +32,14 @@ class USoundCueTemplateFactory : public UFactory
 	UPROPERTY(EditAnywhere, Category = SoundCueTemplateFacotry)
 	TSubclassOf<USoundCueTemplate> SoundCueTemplateClass;
 
+	UPROPERTY(EditAnywhere, Category = SoundCueTemplateFactory)
+	TArray<TWeakObjectPtr<USoundWave>> SoundWaves;
+
 	//~ Begin UFactory Interface
 	virtual bool ConfigureProperties() override;
 	virtual UObject* FactoryCreateNew(UClass* Class, UObject* InParent, FName Name, EObjectFlags Flags, UObject* Context, FFeedbackContext* Warn) override;
 	virtual FString GetDefaultNewAssetName() const override;
 	//~ End UFactory Interface
-
-	TArray<TWeakObjectPtr<USoundWave>> SoundWaves;
 };
 
 UCLASS()
