Fixing SCT asset action that got changed during a refactor.

#rnx

diff --git a/Engine/Plugins/Runtime/SoundCueTemplates/Source/SoundCueTemplatesEditor/Private/AssetTypeActions_SoundCueTemplate.cpp b/Engine/Plugins/Runtime/SoundCueTemplates/Source/SoundCueTemplatesEditor/Private/AssetTypeActions_SoundCueTemplate.cpp
--- a/Engine/Plugins/Runtime/SoundCueTemplates/Source/SoundCueTemplatesEditor/Private/AssetTypeActions_SoundCueTemplate.cpp
+++ b/Engine/Plugins/Runtime/SoundCueTemplates/Source/SoundCueTemplatesEditor/Private/AssetTypeActions_SoundCueTemplate.cpp
@@ -42,7 +42,7 @@ void FAssetActionExtender_SoundCueTemplate::RegisterMenus()
 	FToolMenuOwnerScoped MenuOwner("SoundCueTemplate");
 
 	{
-		UToolMenu* Menu = UToolMenus::Get()->ExtendMenu("ContentBrowser.AssetContextMenu.SoundCue");
+		UToolMenu* Menu = UToolMenus::Get()->ExtendMenu("ContentBrowser.AssetContextMenu.SoundWave");
 		FToolMenuSection& Section = Menu->FindOrAddSection("GetAssetActions");
 
 		Section.AddDynamicEntry("SoundCueAsset", FNewToolMenuSectionDelegate::CreateLambda([](FToolMenuSection& InSection)
