[MetaSounds] - Create SoundCueTemplate appears on MetaSound Source when the Sound Cue Templates plugin is enabled and does nothing when selected

#rnx

diff --git a/Engine/Plugins/Runtime/SoundCueTemplates/Source/SoundCueTemplatesEditor/Private/AssetDefinition_SoundCueTemplate.cpp b/Engine/Plugins/Runtime/SoundCueTemplates/Source/SoundCueTemplatesEditor/Private/AssetDefinition_SoundCueTemplate.cpp
--- a/Engine/Plugins/Runtime/SoundCueTemplates/Source/SoundCueTemplatesEditor/Private/AssetDefinition_SoundCueTemplate.cpp
+++ b/Engine/Plugins/Runtime/SoundCueTemplates/Source/SoundCueTemplatesEditor/Private/AssetDefinition_SoundCueTemplate.cpp
@@ -83,6 +83,21 @@ namespace MenuExtension_SoundCueTemplate
 		}
 	}
 
+	bool CanCreateSoundCueTemplate(const FToolMenuContext& MenuContext)
+	{
+		if (const UContentBrowserAssetContextMenuContext* Context = UContentBrowserAssetContextMenuContext::FindContextWithAssets(MenuContext))
+		{
+			const TArray<TSoftObjectPtr<USoundWaveProcedural>> SoundWaveProcedurals = Context->GetSelectedAssetSoftObjects<USoundWaveProcedural>();
+
+			if (SoundWaveProcedurals.IsEmpty())
+			{
+				return true;
+			}
+		}
+
+		return false;
+	}
+
 	static FDelayedAutoRegisterHelper DelayedAutoRegister(EDelayedRegisterRunPhase::EndOfEngineInit, []{ 
 		UToolMenus::RegisterStartupCallback(FSimpleMulticastDelegate::FDelegate::CreateLambda([]()
 		{
@@ -100,6 +115,7 @@ namespace MenuExtension_SoundCueTemplate
 					
 					FToolUIAction UIAction;
 					UIAction.ExecuteAction = FToolMenuExecuteAction::CreateStatic(&ExecuteCreateSoundCueTemplate);
+					UIAction.IsActionVisibleDelegate = FToolMenuIsActionButtonVisible::CreateStatic(&CanCreateSoundCueTemplate);
 					InSection.AddMenuEntry("SoundWave_CreateSoundCueTemplate", Label, ToolTip, Icon, UIAction);
 				}));
 			}
